PropertyAccessEditor: Allow to use array of structs as binding context.

diff --git a/Engine/Plugins/Runtime/PropertyAccess/Source/PropertyAccessEditor/Private/PropertyAccessEditorModule.cpp b/Engine/Plugins/Runtime/PropertyAccess/Source/PropertyAccessEditor/Private/PropertyAccessEditorModule.cpp
--- a/Engine/Plugins/Runtime/PropertyAccess/Source/PropertyAccessEditor/Private/PropertyAccessEditorModule.cpp
+++ b/Engine/Plugins/Runtime/PropertyAccess/Source/PropertyAccessEditor/Private/PropertyAccessEditorModule.cpp
@@ -37,7 +37,14 @@ public:
 	// IPropertyAccessEditor interface
 	virtual TSharedRef<SWidget> MakePropertyBindingWidget(UBlueprint* InBlueprint, const FPropertyBindingWidgetArgs& InArgs) const override
 	{
-		return SNew(SPropertyBinding, InBlueprint)
+		TArray<FBindingContextStruct> BindingContextStructs;
+		return SNew(SPropertyBinding, InBlueprint, BindingContextStructs)
+			.Args(InArgs);
+	}
+
+	virtual TSharedRef<SWidget> MakePropertyBindingWidget(const TArray<FBindingContextStruct>& InBindingContextStructs, const FPropertyBindingWidgetArgs& InArgs) const override
+	{
+		return SNew(SPropertyBinding, nullptr, InBindingContextStructs)
 			.Args(InArgs);
 	}
 
diff --git a/Engine/Plugins/Runtime/PropertyAccess/Source/PropertyAccessEditor/Private/SPropertyBinding.cpp b/Engine/Plugins/Runtime/PropertyAccess/Source/PropertyAccessEditor/Private/SPropertyBinding.cpp
--- a/Engine/Plugins/Runtime/PropertyAccess/Source/PropertyAccessEditor/Private/SPropertyBinding.cpp
+++ b/Engine/Plugins/Runtime/PropertyAccess/Source/PropertyAccessEditor/Private/SPropertyBinding.cpp
@@ -31,9 +31,10 @@
 /////////////////////////////////////////////////////
 // SPropertyBinding
 
-void SPropertyBinding::Construct(const FArguments& InArgs, UBlueprint* InBlueprint)
+void SPropertyBinding::Construct(const FArguments& InArgs, UBlueprint* InBlueprint, const TArray<FBindingContextStruct>& InBindingContextStructs)
 {
 	Blueprint = InBlueprint;
+	BindingContextStructs = InBindingContextStructs;
 
 	Args = InArgs._Args;
 	PropertyName = Args.Property != nullptr ? Args.Property->GetFName() : NAME_None;
@@ -164,7 +165,7 @@ void SPropertyBinding::ForEachBindableProperty(UStruct* InStruct, Predicate Pred
 {
 	if(Args.OnCanBindProperty.IsBound())
 	{
-		UBlueprintGeneratedClass* SkeletonClass = Cast<UBlueprintGeneratedClass>(Blueprint->SkeletonGeneratedClass);
+		UBlueprintGeneratedClass* SkeletonClass = Blueprint ? Cast<UBlueprintGeneratedClass>(Blueprint->SkeletonGeneratedClass) : nullptr;
 
 		for ( TFieldIterator<FProperty> PropIt(InStruct, EFieldIteratorFlags::IncludeSuper); PropIt; ++PropIt )
 		{
@@ -176,9 +177,12 @@ void SPropertyBinding::ForEachBindableProperty(UStruct* InStruct, Predicate Pred
 				break;
 			}
 
-			if ( !UEdGraphSchema_K2::CanUserKismetAccessVariable(Property, SkeletonClass, UEdGraphSchema_K2::CannotBeDelegate) )
+			if (SkeletonClass)
 			{
-				continue;
+				if (!UEdGraphSchema_K2::CanUserKismetAccessVariable(Property, SkeletonClass, UEdGraphSchema_K2::CannotBeDelegate))
+				{
+					continue;
+				}
 			}
 
 			// Also ignore advanced properties
@@ -224,10 +228,60 @@ TSharedRef<SWidget> SPropertyBinding::OnGenerateDelegateMenu()
 	// Properties
 	{
 		// Get the current skeleton class, think header for the blueprint.
-		UBlueprintGeneratedClass* SkeletonClass = Cast<UBlueprintGeneratedClass>(Blueprint->SkeletonGeneratedClass);
+		UBlueprintGeneratedClass* SkeletonClass = Blueprint ? Cast<UBlueprintGeneratedClass>(Blueprint->SkeletonGeneratedClass) : nullptr;
 
-		TArray<TSharedPtr<FBindingChainElement>> BindingChain;
-		FillPropertyMenu(MenuBuilder, SkeletonClass, BindingChain);
+		if (SkeletonClass)
+		{
+			TArray<TSharedPtr<FBindingChainElement>> BindingChain;
+			FillPropertyMenu(MenuBuilder, SkeletonClass, BindingChain);
+		}
+	}
+
+	if (BindingContextStructs.Num() > 0)
+	{
+		auto MakeContextStructWidget = [this](const FBindingContextStruct& ContextStruct)
+		{
+			FText DisplayText = ContextStruct.DisplayText.IsEmpty() ? ContextStruct.Struct->GetDisplayNameText() : ContextStruct.DisplayText;
+			FText ToolTipText = ContextStruct.TooltipText.IsEmpty() ? DisplayText : ContextStruct.TooltipText;
+
+			return SNew(SHorizontalBox)
+				.ToolTipText(ToolTipText)
+				+ SHorizontalBox::Slot()
+				.AutoWidth()
+				[
+					SNew(SSpacer)
+					.Size(FVector2D(18.0f, 0.0f))
+				]
+				+ SHorizontalBox::Slot()
+				.AutoWidth()
+				.VAlign(VAlign_Center)
+				.Padding(1.0f, 0.0f)
+				[
+					SNew(SImage)
+					.Image(ContextStruct.Icon)
+				]
+				+ SHorizontalBox::Slot()
+				.AutoWidth()
+				.VAlign(VAlign_Center)
+				.Padding(4.0f, 0.0f)
+				[
+					SNew(STextBlock)
+					.Text(DisplayText)
+				];
+		};
+
+		for (int32 i = 0; i < BindingContextStructs.Num(); i++)
+		{
+			FBindingContextStruct& ContextStruct = BindingContextStructs[i];
+
+			// Make first chain element representing the index in the context array.
+			TArray<TSharedPtr<FBindingChainElement>> BindingChain;
+			BindingChain.Emplace(MakeShared<FBindingChainElement>(nullptr, i));
+
+			MenuBuilder.AddSubMenu(
+				MakeContextStructWidget(ContextStruct),
+				FNewMenuDelegate::CreateSP(this, &SPropertyBinding::FillPropertyMenu, ContextStruct.Struct, BindingChain));
+		}
 	}
 
 	FDisplayMetrics DisplayMetrics;
@@ -436,9 +490,6 @@ void SPropertyBinding::FillPropertyMenu(FMenuBuilder& MenuBuilder, UStruct* InOw
 	//---------------------------------------
 	// Property Bindings
 
-	// Get the current skeleton class, think header for the blueprint.
-	UBlueprintGeneratedClass* SkeletonClass = Cast<UBlueprintGeneratedClass>(Blueprint->SkeletonGeneratedClass);
-
 	// Only show bindable subobjects and variables if we're generating pure bindings.
 	if ( Args.bGeneratePureBindings )
 	{
@@ -456,9 +507,21 @@ void SPropertyBinding::FillPropertyMenu(FMenuBuilder& MenuBuilder, UStruct* InOw
 		// to look for bindings that we don't support
 		if ( Args.OnCanBindProperty.IsBound() && Args.OnCanBindProperty.Execute(BindingProperty) )
 		{
+			UStruct* BindingStruct = nullptr;
+			FObjectPropertyBase* ObjectProperty = CastField<FObjectPropertyBase>(BindingProperty);
+			FStructProperty* StructProperty = CastField<FStructProperty>(BindingProperty);
+			if (ObjectProperty)
+			{
+				BindingStruct = ObjectProperty->PropertyClass;
+			}
+			else if (StructProperty)
+			{
+				BindingStruct = StructProperty->Struct;
+			}
+
 			MenuBuilder.BeginSection("Properties", LOCTEXT("Properties", "Properties"));
 			{
-				ForEachBindableProperty(InOwnerStruct, [this, &InBindingChain, &MenuBuilder, &MakeArrayElementPropertyWidget, &MakePropertyWidget, &MakePropertyEntry, &MakeArrayElementEntry] (FProperty* Property) 
+				ForEachBindableProperty(InOwnerStruct, [this, &InBindingChain, BindingStruct, &MenuBuilder, &MakeArrayElementPropertyWidget, &MakePropertyWidget, &MakePropertyEntry, &MakeArrayElementEntry] (FProperty* Property)
 				{
 					TArray<TSharedPtr<FBindingChainElement>> NewBindingChain(InBindingChain);
 					NewBindingChain.Emplace(MakeShared<FBindingChainElement>(Property));
@@ -500,9 +563,10 @@ void SPropertyBinding::FillPropertyMenu(FMenuBuilder& MenuBuilder, UStruct* InOw
 						Struct = StructProperty->Struct;
 					}
 
-					if ( Struct )
+					// Recurse into a struct, except if it is the same type as the one we're binding.
+					if (Struct && Struct != BindingStruct)
 					{
-						if ( Class )
+						if (Class)
 						{
 							// Ignore any subobject properties that are not bindable.
 							// Also ignore any class that is explicitly on the black list.
@@ -512,7 +576,7 @@ void SPropertyBinding::FillPropertyMenu(FMenuBuilder& MenuBuilder, UStruct* InOw
 							}
 						}
 
-						if(Args.bAllowArrayElementBindings && ArrayProperty != nullptr)
+						if (Args.bAllowArrayElementBindings && ArrayProperty != nullptr)
 						{
 							TArray<TSharedPtr<FBindingChainElement>> NewArrayElementBindingChain(InBindingChain);
 							NewArrayElementBindingChain.Emplace(MakeShared<FBindingChainElement>(Property));
@@ -623,6 +687,11 @@ void SPropertyBinding::HandleSetBindingArrayIndex(int32 InArrayIndex, ETextCommi
 
 void SPropertyBinding::HandleCreateAndAddBinding()
 {
+	if (!Blueprint)
+	{
+		return;
+	}
+
 	const FScopedTransaction Transaction(LOCTEXT("CreateDelegate", "Create Binding"));
 
 	Blueprint->Modify();
diff --git a/Engine/Plugins/Runtime/PropertyAccess/Source/PropertyAccessEditor/Private/SPropertyBinding.h b/Engine/Plugins/Runtime/PropertyAccess/Source/PropertyAccessEditor/Private/SPropertyBinding.h
--- a/Engine/Plugins/Runtime/PropertyAccess/Source/PropertyAccessEditor/Private/SPropertyBinding.h
+++ b/Engine/Plugins/Runtime/PropertyAccess/Source/PropertyAccessEditor/Private/SPropertyBinding.h
@@ -26,7 +26,7 @@ public:
 
 	SLATE_END_ARGS()
 
-	void Construct(const FArguments& InArgs, UBlueprint* InBlueprint);
+	void Construct(const FArguments& InArgs, UBlueprint* InBlueprint, const TArray<FBindingContextStruct>& InBindingContextStructs);
 
 protected:
 	struct FFunctionInfo
@@ -80,6 +80,7 @@ private:
 	void ForEachBindableFunction(UClass* FromClass, Predicate Pred) const;
 
 	UBlueprint* Blueprint;
+	TArray<FBindingContextStruct> BindingContextStructs;
 	FPropertyBindingWidgetArgs Args;
 	FName PropertyName;
 };
diff --git a/Engine/Source/Editor/UnrealEd/Public/IPropertyAccessEditor.h b/Engine/Source/Editor/UnrealEd/Public/IPropertyAccessEditor.h
--- a/Engine/Source/Editor/UnrealEd/Public/IPropertyAccessEditor.h
+++ b/Engine/Source/Editor/UnrealEd/Public/IPropertyAccessEditor.h
@@ -59,6 +59,33 @@ struct FRedirectorBindingInfo
 	UStruct* Struct = nullptr;
 };
 
+/**
+ * Binding context struct allow to describe information for a struct to bind to using the binding widget. An array of structs is passed to the widget to describe the context in which the binding exists.
+ * When the widget selectes a property from binding context struct array, the first FBindingChainElement's index correlates to the array passed to the widget.
+ */
+struct FBindingContextStruct
+{
+	FBindingContextStruct(UStruct* InStruct, const FSlateBrush* InIcon = nullptr, const FText& InDisplayText = FText::GetEmpty(), const FText& InTooltipText = FText::GetEmpty())
+		: Struct(InStruct)
+		, Icon(InIcon)
+		, DisplayText(InDisplayText)
+		, TooltipText(InTooltipText)
+	{}
+
+	/** The struct to bind to. */
+	UStruct* Struct;
+
+	/** Icon to display in hte poppup menu. */ 
+	const FSlateBrush* Icon;
+
+	/** Text to display for this item in the popup. If left empty, struct's display text will be used. */
+	FText DisplayText;
+
+	/** Tooltip to show, or if empty, the tool tip will be set to the same as popup text. */
+	FText TooltipText;
+};
+
+
 /** Delegate used to generate a new binding function's name */
 DECLARE_DELEGATE_RetVal(FString, FOnGenerateBindingName);
 
@@ -193,6 +220,14 @@ public:
 	 */
 	virtual TSharedRef<SWidget> MakePropertyBindingWidget(UBlueprint* InBlueprint, const FPropertyBindingWidgetArgs& InArgs = FPropertyBindingWidgetArgs()) const = 0;
 
+	/**
+	 * Make a property binding widget.
+	 * @param	InBindingContextStructs		An array of structs the binding will exist within
+	 * @param	InArgs						Optional arguments for the widget
+	 * @return a new binding widget
+	 */
+	virtual TSharedRef<SWidget> MakePropertyBindingWidget(const TArray<FBindingContextStruct>& InBindingContextStructs, const FPropertyBindingWidgetArgs& InArgs = FPropertyBindingWidgetArgs()) const = 0;
+
 	/** Resolve a property path to a structure, returning the leaf property and array index if any. @return true if resolution succeeded */
 	virtual EPropertyAccessResolveResult ResolveLeafProperty(const UStruct* InStruct, TArrayView<FString> InPath, FProperty*& OutProperty, int32& OutArrayIndex) const = 0;
 
