Fixed ensure and infinitely-growing assets when cooking property access paths

The property access library was never cleared to its default state, so each time a compilation occured the library grew in size, preserving all of its older (potentially out of date) values. This re-instates a call to BeginCompilation() that was removed as part of a refactor.

diff --git a/Engine/Plugins/Runtime/PropertyAccess/Source/PropertyAccessEditor/Private/AnimBlueprintCompilerHandler_PropertyAccess.cpp b/Engine/Plugins/Runtime/PropertyAccess/Source/PropertyAccessEditor/Private/AnimBlueprintCompilerHandler_PropertyAccess.cpp
--- a/Engine/Plugins/Runtime/PropertyAccess/Source/PropertyAccessEditor/Private/AnimBlueprintCompilerHandler_PropertyAccess.cpp
+++ b/Engine/Plugins/Runtime/PropertyAccess/Source/PropertyAccessEditor/Private/AnimBlueprintCompilerHandler_PropertyAccess.cpp
@@ -21,6 +21,7 @@ int32 FAnimBlueprintCompilerHandler_PropertyAccess::AddCopy(TArrayView<FString>
 void FAnimBlueprintCompilerHandler_PropertyAccess::StartCompilingClass(const UClass* InClass, IAnimBlueprintCompilationBracketContext& InCompilationContext, IAnimBlueprintGeneratedClassCompiledData& OutCompiledData)
 {
 	PropertyAccessLibraryCompiler.Setup(InClass, &OutCompiledData.GetPropertyAccessLibrary());
+	PropertyAccessLibraryCompiler.BeginCompilation(InClass);
 }
 
 void FAnimBlueprintCompilerHandler_PropertyAccess::FinishCompilingClass(const UClass* InClass, IAnimBlueprintCompilationBracketContext& InCompilationContext, IAnimBlueprintGeneratedClassCompiledData& OutCompiledData)
