PropertyAccess / PropertyBinding to support filtering of struct members

diff --git a/Engine/Plugins/Runtime/PropertyAccess/Source/PropertyAccessEditor/Private/SPropertyBinding.cpp b/Engine/Plugins/Runtime/PropertyAccess/Source/PropertyAccessEditor/Private/SPropertyBinding.cpp
--- a/Engine/Plugins/Runtime/PropertyAccess/Source/PropertyAccessEditor/Private/SPropertyBinding.cpp
+++ b/Engine/Plugins/Runtime/PropertyAccess/Source/PropertyAccessEditor/Private/SPropertyBinding.cpp
@@ -587,7 +587,7 @@ void SPropertyBinding::FillPropertyMenu(FMenuBuilder& MenuBuilder, UStruct* InOw
 								FNewMenuDelegate::CreateSP(this, &SPropertyBinding::FillPropertyMenu, Struct, NewArrayElementBindingChain)
 								);
 						}
-						else
+						else if(Args.bAllowStructMemberBindings)
 						{
 							MenuBuilder.AddSubMenu(
 								MakePropertyWidget(Property),
diff --git a/Engine/Source/Editor/UnrealEd/Public/IPropertyAccessEditor.h b/Engine/Source/Editor/UnrealEd/Public/IPropertyAccessEditor.h
--- a/Engine/Source/Editor/UnrealEd/Public/IPropertyAccessEditor.h
+++ b/Engine/Source/Editor/UnrealEd/Public/IPropertyAccessEditor.h
@@ -173,6 +173,9 @@ struct FPropertyBindingWidgetArgs
 	/** Whether to allow array element bindings */
 	bool bAllowArrayElementBindings = false;
 
+	/** Whether to allow struct member bindings */
+	bool bAllowStructMemberBindings = true;
+
 	/** Whether to allow new bindings to be made from within the widget's UI */
 	bool bAllowNewBindings = true;
 
