PropertyAccessEditor
- Added pruning of empty menus in case child classes/structs dont have anything to bind to
- Added OnCanAcceptProperty which can be used to skip properties completelly

diff --git a/Engine/Plugins/Runtime/PropertyAccess/Source/PropertyAccessEditor/Private/SPropertyBinding.cpp b/Engine/Plugins/Runtime/PropertyAccess/Source/PropertyAccessEditor/Private/SPropertyBinding.cpp
--- a/Engine/Plugins/Runtime/PropertyAccess/Source/PropertyAccessEditor/Private/SPropertyBinding.cpp
+++ b/Engine/Plugins/Runtime/PropertyAccess/Source/PropertyAccessEditor/Private/SPropertyBinding.cpp
@@ -183,6 +183,11 @@ void SPropertyBinding::ForEachBindableProperty(UStruct* InStruct, Predicate Pred
 				break;
 			}
 
+			if (Args.OnCanAcceptPropertyOrChildren.IsBound() && Args.OnCanAcceptPropertyOrChildren.Execute(Property) == false)
+			{
+				break;
+			}
+
 			if (SkeletonClass)
 			{
 				if (!UEdGraphSchema_K2::CanUserKismetAccessVariable(Property, SkeletonClass, UEdGraphSchema_K2::CannotBeDelegate))
@@ -192,7 +197,7 @@ void SPropertyBinding::ForEachBindableProperty(UStruct* InStruct, Predicate Pred
 			}
 
 			// Also ignore advanced properties
-			if ( Property->HasAnyPropertyFlags(CPF_AdvancedDisplay | CPF_EditorOnly) )
+			if (Property->HasAnyPropertyFlags(CPF_AdvancedDisplay | CPF_EditorOnly))
 			{
 				continue;
 			}
@@ -202,6 +207,73 @@ void SPropertyBinding::ForEachBindableProperty(UStruct* InStruct, Predicate Pred
 	}
 }
 
+bool SPropertyBinding::HasBindableProperties(UStruct* InStruct) const
+{
+	int32 BindableCount = 0;
+	ForEachBindableProperty(InStruct, [this, &BindableCount] (FProperty* Property)
+	{
+		FArrayProperty* ArrayProperty = CastField<FArrayProperty>(Property);
+
+		if(Args.OnCanBindProperty.Execute(Property))
+		{
+			BindableCount++;
+			return;
+		}
+
+		if(Args.bAllowArrayElementBindings && ArrayProperty != nullptr && Args.OnCanBindProperty.Execute(ArrayProperty->Inner))
+		{
+			BindableCount++;
+			return;
+		}
+
+		FProperty* InnerProperty = Property;
+		if(Args.bAllowArrayElementBindings && ArrayProperty != nullptr)
+		{
+			InnerProperty = ArrayProperty->Inner;
+		}
+
+		FObjectPropertyBase* ObjectProperty = CastField<FObjectPropertyBase>(InnerProperty);
+		FStructProperty* StructProperty = CastField<FStructProperty>(InnerProperty);
+
+		UStruct* Struct = nullptr;
+		UClass* Class = nullptr;
+
+		if ( ObjectProperty )
+		{
+			Struct = Class = ObjectProperty->PropertyClass;
+		}
+		else if ( StructProperty )
+		{
+			Struct = StructProperty->Struct;
+		}
+
+		// Recurse into a struct, except if it is the same type as the one we're binding.
+		if (Struct && HasBindableProperties(Struct))
+		{
+			if (Class)
+			{
+				// Ignore any subobject properties that are not bindable.
+				// Also ignore any class that is explicitly on the black list.
+				if ( IsClassBlackListed(Class) || (Args.OnCanBindToSubObjectClass.IsBound() && Args.OnCanBindToSubObjectClass.Execute(Class)))
+				{
+					return;
+				}
+			}
+
+			if (Args.bAllowArrayElementBindings && ArrayProperty != nullptr)
+			{
+				BindableCount++;
+			}
+			else if (Args.bAllowStructMemberBindings)
+			{
+				BindableCount++;
+			}
+		}
+	});
+	
+	return BindableCount > 0;				
+}
+
 TSharedRef<SWidget> SPropertyBinding::OnGenerateDelegateMenu()
 {
 	const bool bInShouldCloseWindowAfterMenuSelection = true;
@@ -279,14 +351,16 @@ TSharedRef<SWidget> SPropertyBinding::OnGenerateDelegateMenu()
 		for (int32 i = 0; i < BindingContextStructs.Num(); i++)
 		{
 			FBindingContextStruct& ContextStruct = BindingContextStructs[i];
+			if (HasBindableProperties(ContextStruct.Struct))
+			{
+				// Make first chain element representing the index in the context array.
+				TArray<TSharedPtr<FBindingChainElement>> BindingChain;
+				BindingChain.Emplace(MakeShared<FBindingChainElement>(nullptr, i));
 
-			// Make first chain element representing the index in the context array.
-			TArray<TSharedPtr<FBindingChainElement>> BindingChain;
-			BindingChain.Emplace(MakeShared<FBindingChainElement>(nullptr, i));
-
-			MenuBuilder.AddSubMenu(
-				MakeContextStructWidget(ContextStruct),
-				FNewMenuDelegate::CreateSP(this, &SPropertyBinding::FillPropertyMenu, ContextStruct.Struct, BindingChain));
+				MenuBuilder.AddSubMenu(
+					MakeContextStructWidget(ContextStruct),
+					FNewMenuDelegate::CreateSP(this, &SPropertyBinding::FillPropertyMenu, ContextStruct.Struct, BindingChain));
+			}
 		}
 	}
 
@@ -578,8 +652,8 @@ void SPropertyBinding::FillPropertyMenu(FMenuBuilder& MenuBuilder, UStruct* InOw
 						Struct = StructProperty->Struct;
 					}
 
-					// Recurse into a struct, except if it is the same type as the one we're binding.
-					if (Struct && Struct != BindingStruct)
+					// Recurse into a struct if it has some properties we can bind to.
+					if (Struct)
 					{
 						if (Class)
 						{
@@ -591,22 +665,25 @@ void SPropertyBinding::FillPropertyMenu(FMenuBuilder& MenuBuilder, UStruct* InOw
 							}
 						}
 
-						if (Args.bAllowArrayElementBindings && ArrayProperty != nullptr)
-						{
-							TArray<TSharedPtr<FBindingChainElement>> NewArrayElementBindingChain(InBindingChain);
-							NewArrayElementBindingChain.Emplace(MakeShared<FBindingChainElement>(Property));
-							NewArrayElementBindingChain.Last()->ArrayIndex = 0;
-
-							MenuBuilder.AddSubMenu(
-								MakeArrayElementPropertyWidget(ArrayProperty->Inner, NewArrayElementBindingChain),
-								FNewMenuDelegate::CreateSP(this, &SPropertyBinding::FillPropertyMenu, Struct, NewArrayElementBindingChain)
-								);
-						}
-						else if(Args.bAllowStructMemberBindings)
+						if (HasBindableProperties(Struct))
 						{
-							MenuBuilder.AddSubMenu(
-								MakePropertyWidget(Property),
-								FNewMenuDelegate::CreateSP(this, &SPropertyBinding::FillPropertyMenu, Struct, NewBindingChain));
+							if (Args.bAllowArrayElementBindings && ArrayProperty != nullptr)
+							{
+								TArray<TSharedPtr<FBindingChainElement>> NewArrayElementBindingChain(InBindingChain);
+								NewArrayElementBindingChain.Emplace(MakeShared<FBindingChainElement>(Property));
+								NewArrayElementBindingChain.Last()->ArrayIndex = 0;
+
+								MenuBuilder.AddSubMenu(
+									MakeArrayElementPropertyWidget(ArrayProperty->Inner, NewArrayElementBindingChain),
+									FNewMenuDelegate::CreateSP(this, &SPropertyBinding::FillPropertyMenu, Struct, NewArrayElementBindingChain)
+									);
+							}
+							else if (Args.bAllowStructMemberBindings)
+							{
+								MenuBuilder.AddSubMenu(
+									MakePropertyWidget(Property),
+									FNewMenuDelegate::CreateSP(this, &SPropertyBinding::FillPropertyMenu, Struct, NewBindingChain));
+							}
 						}
 					}
 				});
diff --git a/Engine/Plugins/Runtime/PropertyAccess/Source/PropertyAccessEditor/Private/SPropertyBinding.h b/Engine/Plugins/Runtime/PropertyAccess/Source/PropertyAccessEditor/Private/SPropertyBinding.h
--- a/Engine/Plugins/Runtime/PropertyAccess/Source/PropertyAccessEditor/Private/SPropertyBinding.h
+++ b/Engine/Plugins/Runtime/PropertyAccess/Source/PropertyAccessEditor/Private/SPropertyBinding.h
@@ -73,7 +73,8 @@ protected:
 private:
 	bool IsClassBlackListed(UClass* OwnerClass) const;
 	bool IsFieldFromBlackListedClass(FFieldVariant Field) const;
-
+	bool HasBindableProperties(UStruct* InStruct) const;
+	
 	template <typename Predicate>
 	void ForEachBindableProperty(UStruct* InStruct, Predicate Pred) const;
 
diff --git a/Engine/Source/Editor/UnrealEd/Public/IPropertyAccessEditor.h b/Engine/Source/Editor/UnrealEd/Public/IPropertyAccessEditor.h
--- a/Engine/Source/Editor/UnrealEd/Public/IPropertyAccessEditor.h
+++ b/Engine/Source/Editor/UnrealEd/Public/IPropertyAccessEditor.h
@@ -97,6 +97,9 @@ DECLARE_DELEGATE_RetVal_OneParam(bool, FOnGotoBinding, FName /*InPropertyName*/)
 /** Delegate used to se if we can open a binding (e.g. a function) */
 DECLARE_DELEGATE_RetVal_OneParam(bool, FOnCanGotoBinding, FName /*InPropertyName*/);
 
+/** Delegate used to check whether a property is considered for binding. Returning false will discard the property and all child properties. */
+DECLARE_DELEGATE_RetVal_OneParam(bool, FOnCanAcceptPropertyOrChildren, FProperty* /*InProperty*/);
+
 /** Delegate used to check whether a property can be bound to the property in question */
 DECLARE_DELEGATE_RetVal_OneParam(bool, FOnCanBindProperty, FProperty* /*InProperty*/);
 
@@ -139,6 +142,9 @@ struct FPropertyBindingWidgetArgs
 	/** Delegate used to see if we can open a binding (e.g. a function) */
 	FOnCanGotoBinding OnCanGotoBinding;
 
+	/** Delegate used to check whether a property is considered for binding. Returning false will discard the property and all child properties. */
+	FOnCanAcceptPropertyOrChildren OnCanAcceptPropertyOrChildren;
+	
 	/** Delegate used to check whether a property can be bound to the property in question */
 	FOnCanBindProperty OnCanBindProperty;
 
