PropertyAccessEditor: Fixed infinite loop when checking if struct can be bound to.

diff --git a/Engine/Plugins/Runtime/PropertyAccess/Source/PropertyAccessEditor/Private/SPropertyBinding.cpp b/Engine/Plugins/Runtime/PropertyAccess/Source/PropertyAccessEditor/Private/SPropertyBinding.cpp
--- a/Engine/Plugins/Runtime/PropertyAccess/Source/PropertyAccessEditor/Private/SPropertyBinding.cpp
+++ b/Engine/Plugins/Runtime/PropertyAccess/Source/PropertyAccessEditor/Private/SPropertyBinding.cpp
@@ -209,8 +209,28 @@ void SPropertyBinding::ForEachBindableProperty(UStruct* InStruct, Predicate Pred
 
 bool SPropertyBinding::HasBindableProperties(UStruct* InStruct) const
 {
+	TSet<UStruct*> VisitedStructs;
+	return HasBindablePropertiesRecursive(InStruct, VisitedStructs, 0);
+}
+
+bool SPropertyBinding::HasBindablePropertiesRecursive(UStruct* InStruct, TSet<UStruct*>& VisitedStructs, const int32 RecursionDepth) const
+{
+	if (VisitedStructs.Contains(InStruct))
+	{
+		// If we have visited this struct already and it had bindable properties, it has been accounted for already.
+		// Returning false to allow early exit.
+		return false;
+	}
+	VisitedStructs.Add(InStruct);
+	
+	// Arbitrary cut off to avoid infinite loops.
+	if (RecursionDepth > 10)
+	{
+		return false;
+	}
+	
 	int32 BindableCount = 0;
-	ForEachBindableProperty(InStruct, [this, &BindableCount] (FProperty* Property)
+	ForEachBindableProperty(InStruct, [this, &BindableCount, &VisitedStructs, RecursionDepth] (FProperty* Property)
 	{
 		FArrayProperty* ArrayProperty = CastField<FArrayProperty>(Property);
 
@@ -248,7 +268,7 @@ bool SPropertyBinding::HasBindableProperties(UStruct* InStruct) const
 		}
 
 		// Recurse into a struct, except if it is the same type as the one we're binding.
-		if (Struct && HasBindableProperties(Struct))
+		if (Struct && HasBindablePropertiesRecursive(Struct, VisitedStructs, RecursionDepth + 1))
 		{
 			if (Class)
 			{
diff --git a/Engine/Plugins/Runtime/PropertyAccess/Source/PropertyAccessEditor/Private/SPropertyBinding.h b/Engine/Plugins/Runtime/PropertyAccess/Source/PropertyAccessEditor/Private/SPropertyBinding.h
--- a/Engine/Plugins/Runtime/PropertyAccess/Source/PropertyAccessEditor/Private/SPropertyBinding.h
+++ b/Engine/Plugins/Runtime/PropertyAccess/Source/PropertyAccessEditor/Private/SPropertyBinding.h
@@ -74,6 +74,7 @@ private:
 	bool IsClassBlackListed(UClass* OwnerClass) const;
 	bool IsFieldFromBlackListedClass(FFieldVariant Field) const;
 	bool HasBindableProperties(UStruct* InStruct) const;
+	bool HasBindablePropertiesRecursive(UStruct* InStruct, TSet<UStruct*>& VisitedStructs, const int32 RecursionDepth) const;
 	
 	template <typename Predicate>
 	void ForEachBindableProperty(UStruct* InStruct, Predicate Pred) const;
