Property Access: fixed use of OnCanAcceptPropertyOrChildren
- OnCanAcceptPropertyOrChildren should skip the property, not stop completely

diff --git a/Engine/Plugins/Runtime/PropertyAccess/Source/PropertyAccessEditor/Private/SPropertyBinding.cpp b/Engine/Plugins/Runtime/PropertyAccess/Source/PropertyAccessEditor/Private/SPropertyBinding.cpp
--- a/Engine/Plugins/Runtime/PropertyAccess/Source/PropertyAccessEditor/Private/SPropertyBinding.cpp
+++ b/Engine/Plugins/Runtime/PropertyAccess/Source/PropertyAccessEditor/Private/SPropertyBinding.cpp
@@ -185,7 +185,7 @@ void SPropertyBinding::ForEachBindableProperty(UStruct* InStruct, Predicate Pred
 
 			if (Args.OnCanAcceptPropertyOrChildren.IsBound() && Args.OnCanAcceptPropertyOrChildren.Execute(Property) == false)
 			{
-				break;
+				continue;
 			}
 
 			if (SkeletonClass)
