Fix warnings about missing variables making property accesses non-functional

Now if we detect an error in a copy we still rebuild the library, although the copy that was in error will not function, as expected

diff --git a/Engine/Plugins/Runtime/PropertyAccess/Source/PropertyAccessEditor/Private/PropertyAccessEditor.cpp b/Engine/Plugins/Runtime/PropertyAccess/Source/PropertyAccessEditor/Private/PropertyAccessEditor.cpp
--- a/Engine/Plugins/Runtime/PropertyAccess/Source/PropertyAccessEditor/Private/PropertyAccessEditor.cpp
+++ b/Engine/Plugins/Runtime/PropertyAccess/Source/PropertyAccessEditor/Private/PropertyAccessEditor.cpp
@@ -658,10 +658,9 @@ bool FPropertyAccessLibraryCompiler::FinishCompilation()
 			CopyMap.Add(CopyIndex, Copy.BatchIndex);
 		}
 
-		if(bResult)
-		{
-			PropertyAccess::PostLoadLibrary(*Library);
-		}
+		// Always rebuild the library even if we detected a 'failure'. Otherwise we could fail to copy data for both
+		// valid and invalid copies 
+		PropertyAccess::PostLoadLibrary(*Library);
 
 		return bResult;
 	}
