Fix property access thread safety logic for functions

Correctly override the thread safety of return-value properties when a function is marked thread-safe.
Previously functions needed to both have a non-uobject return value AND be thread safe for the access to be marked as thread-safe overall. Now the logic correctly overrides the (usually not thread safe) object return value with the function metadata.

diff --git a/Engine/Plugins/Runtime/PropertyAccess/Source/PropertyAccessEditor/Private/PropertyAccessEditor.cpp b/Engine/Plugins/Runtime/PropertyAccess/Source/PropertyAccessEditor/Private/PropertyAccessEditor.cpp
--- a/Engine/Plugins/Runtime/PropertyAccess/Source/PropertyAccessEditor/Private/PropertyAccessEditor.cpp
+++ b/Engine/Plugins/Runtime/PropertyAccess/Source/PropertyAccessEditor/Private/PropertyAccessEditor.cpp
@@ -58,7 +58,7 @@ struct FPropertyAccessEditorSystem
 		Succeeded,
 	};
 
-	static ESegmentResolveResult ResolveSegments_CheckProperty(FPropertyAccessSegment& InSegment, FProperty* InProperty, FResolveSegmentsContext& InContext)
+	static ESegmentResolveResult ResolveSegments_CheckProperty(FPropertyAccessSegment& InSegment, FProperty* InProperty, FResolveSegmentsContext& InContext, bool& bOutThreadSafe)
 	{
 		InSegment.Property = InProperty;
 
@@ -79,7 +79,7 @@ struct FPropertyAccessEditorSystem
 				InSegment.Struct = ArrayOfObjectsProperty->PropertyClass;
 				if(!InContext.bFinalSegment)
 				{
-					InContext.bWasThreadSafe = false;
+					bOutThreadSafe = false;
 				}
 			}
 			else
@@ -106,7 +106,7 @@ struct FPropertyAccessEditorSystem
 			InSegment.Struct = ObjectProperty->PropertyClass;
 			if(!InContext.bFinalSegment)
 			{
-				InContext.bWasThreadSafe = false;
+				bOutThreadSafe = false;
 			}
 		}
 		// Check to see if this is a simple weak object property (eg. not an array of weak objects).
@@ -116,7 +116,7 @@ struct FPropertyAccessEditorSystem
 			InSegment.Struct = WeakObjectProperty->PropertyClass;
 			if(!InContext.bFinalSegment)
 			{
-				InContext.bWasThreadSafe = false;
+				bOutThreadSafe = false;
 			}
 		}
 		// Check to see if this is a simple soft object property (eg. not an array of soft objects).
@@ -126,7 +126,7 @@ struct FPropertyAccessEditorSystem
 			InSegment.Struct = SoftObjectProperty->PropertyClass;
 			if(!InContext.bFinalSegment)
 			{
-				InContext.bWasThreadSafe = false;
+				bOutThreadSafe = false;
 			}
 		}
 		else
@@ -158,7 +158,8 @@ struct FPropertyAccessEditorSystem
 		}
 
 		// Treat the function's return value as the struct/class we want to use for the next segment
-		const ESegmentResolveResult Result = ResolveSegments_CheckProperty(InSegment, ReturnProperty, InContext);
+		bool bThreadSafeProperty = true;
+		const ESegmentResolveResult Result = ResolveSegments_CheckProperty(InSegment, ReturnProperty, InContext, bThreadSafeProperty);
 		if(Result != ESegmentResolveResult::Failed)
 		{
 			// Check a function's thread safety.
@@ -166,7 +167,7 @@ struct FPropertyAccessEditorSystem
 			// can be overridden here if the function that returns the value promises that it is thread safe to access that object.
 			// An example of this is would be something like accessing the main anim BP from a linked anim BP, where it is 'safe' to access
 			// the other object while running animation updated on a worker thread.
-			InContext.bWasThreadSafe &= FBlueprintEditorUtils::HasFunctionBlueprintThreadSafeMetaData(InFunction);
+			InContext.bWasThreadSafe &= (bThreadSafeProperty || FBlueprintEditorUtils::HasFunctionBlueprintThreadSafeMetaData(InFunction));
 		}
 
 		return Result;
@@ -212,7 +213,7 @@ struct FPropertyAccessEditorSystem
 
 				if(FProperty* Property = Field.Get<FProperty>())
 				{
-					if(ResolveSegments_CheckProperty(Segment, Property, InContext) == ESegmentResolveResult::Failed)
+					if(ResolveSegments_CheckProperty(Segment, Property, InContext, InContext.bWasThreadSafe) == ESegmentResolveResult::Failed)
 					{
 						return EPropertyAccessResolveResult::Failed;
 					}
