Fix property access thread safe logic to correctly enforce function metadata overrides

diff --git a/Engine/Plugins/Runtime/PropertyAccess/Source/PropertyAccessEditor/Private/PropertyAccessEditor.cpp b/Engine/Plugins/Runtime/PropertyAccess/Source/PropertyAccessEditor/Private/PropertyAccessEditor.cpp
--- a/Engine/Plugins/Runtime/PropertyAccess/Source/PropertyAccessEditor/Private/PropertyAccessEditor.cpp
+++ b/Engine/Plugins/Runtime/PropertyAccess/Source/PropertyAccessEditor/Private/PropertyAccessEditor.cpp
@@ -167,7 +167,8 @@ struct FPropertyAccessEditorSystem
 			// can be overridden here if the function that returns the value promises that it is thread safe to access that object.
 			// An example of this is would be something like accessing the main anim BP from a linked anim BP, where it is 'safe' to access
 			// the other object while running animation updated on a worker thread.
-			InContext.bWasThreadSafe &= (bThreadSafeProperty || FBlueprintEditorUtils::HasFunctionBlueprintThreadSafeMetaData(InFunction));
+			const bool bThreadSafeFunction = FBlueprintEditorUtils::HasFunctionBlueprintThreadSafeMetaData(InFunction); 
+			InContext.bWasThreadSafe &= (bThreadSafeProperty && bThreadSafeFunction) || (!bThreadSafeProperty && bThreadSafeFunction);
 		}
 
 		return Result;
