Property access UI fixes

Nested property validity filter (checking to see if an entry in the menu has compatible sub-menus) is now performed for functions.
Prevented UObject types from being bound to unrelated types via the menu (now does a IsChildOf check - previously you could attempt to bind a UAnimSequence to an AActor!)
Display 'display name' for script properties/functions (also use ScriptName for native ones if available).

diff --git a/Engine/Plugins/Developer/PropertyAccessNode/Source/PropertyAccessNode/Private/K2Node_PropertyAccess.cpp b/Engine/Plugins/Developer/PropertyAccessNode/Source/PropertyAccessNode/Private/K2Node_PropertyAccess.cpp
--- a/Engine/Plugins/Developer/PropertyAccessNode/Source/PropertyAccessNode/Private/K2Node_PropertyAccess.cpp
+++ b/Engine/Plugins/Developer/PropertyAccessNode/Source/PropertyAccessNode/Private/K2Node_PropertyAccess.cpp
@@ -126,7 +126,7 @@ void UK2Node_PropertyAccess::SetPath(const TArray<FString>& InPath)
 	Path = InPath;
 
 	IPropertyAccessEditor& PropertyAccessEditor = IModularFeatures::Get().GetModularFeature<IPropertyAccessEditor>("PropertyAccessEditor");
-	TextPath = PropertyAccessEditor.MakeTextPath(Path);
+	TextPath = PropertyAccessEditor.MakeTextPath(Path, GetBlueprint()->SkeletonGeneratedClass);
 	FBlueprintEditorUtils::MarkBlueprintAsStructurallyModified(GetBlueprint());
 	ResolvePropertyAccess();
 	ReconstructNode();
@@ -137,7 +137,7 @@ void UK2Node_PropertyAccess::SetPath(TArray<FString>&& InPath)
 	Path = MoveTemp(InPath);
 	
 	IPropertyAccessEditor& PropertyAccessEditor = IModularFeatures::Get().GetModularFeature<IPropertyAccessEditor>("PropertyAccessEditor");
-	TextPath = PropertyAccessEditor.MakeTextPath(Path);
+	TextPath = PropertyAccessEditor.MakeTextPath(Path, GetBlueprint()->SkeletonGeneratedClass);
 	FBlueprintEditorUtils::MarkBlueprintAsStructurallyModified(GetBlueprint());
 	ResolvePropertyAccess();
 	ReconstructNode();
@@ -160,6 +160,10 @@ void UK2Node_PropertyAccess::AllocatePins(UEdGraphPin* InOldOutputPin)
 
 	if(UBlueprint* Blueprint = GetBlueprint())
 	{
+		// Keep text path up to date
+		IPropertyAccessEditor& PropertyAccessEditor = IModularFeatures::Get().GetModularFeature<IPropertyAccessEditor>("PropertyAccessEditor");
+		TextPath = PropertyAccessEditor.MakeTextPath(Path, GetBlueprint()->SkeletonGeneratedClass);
+		
 		UEdGraphPin* OutputPin = nullptr;
 	
 		if(InOldOutputPin != nullptr && InOldOutputPin->LinkedTo.Num() > 0 && InOldOutputPin->PinType.PinCategory != UEdGraphSchema_K2::PC_Wildcard)
@@ -318,7 +322,7 @@ void UK2Node_PropertyAccess::HandleVariableRenamed(UBlueprint* InBlueprint, UCla
 	for(const int32& RenameIndex : RenameIndices)
 	{
 		Path[RenameIndex] = InNewVarName.ToString();
-		TextPath = PropertyAccessEditor.MakeTextPath(Path);
+		TextPath = PropertyAccessEditor.MakeTextPath(Path, GetBlueprint()->SkeletonGeneratedClass);
 	}
 }
 
@@ -350,7 +354,7 @@ void UK2Node_PropertyAccess::ReplaceReferences(UBlueprint* InBlueprint, UBluepri
 	for(const int32& RenameIndex : ReplaceIndices)
 	{
 		Path[RenameIndex] = ReplacementProperty->GetName();
-		TextPath = PropertyAccessEditor.MakeTextPath(Path);
+		TextPath = PropertyAccessEditor.MakeTextPath(Path, GetBlueprint()->SkeletonGeneratedClass);
 	}
 }
 
diff --git a/Engine/Plugins/Developer/PropertyAccessNode/Source/PropertyAccessNode/Private/SPropertyAccessNode.cpp b/Engine/Plugins/Developer/PropertyAccessNode/Source/PropertyAccessNode/Private/SPropertyAccessNode.cpp
--- a/Engine/Plugins/Developer/PropertyAccessNode/Source/PropertyAccessNode/Private/SPropertyAccessNode.cpp
+++ b/Engine/Plugins/Developer/PropertyAccessNode/Source/PropertyAccessNode/Private/SPropertyAccessNode.cpp
@@ -154,15 +154,18 @@ TSharedRef<SWidget> SPropertyAccessNode::UpdateTitleWidget(FText InTitleText, TS
 		}
 		else
 		{
+			IPropertyAccessEditor& PropertyAccessEditor = IModularFeatures::Get().GetModularFeature<IPropertyAccessEditor>("PropertyAccessEditor");
+			const FText UnderlyingPath = PropertyAccessEditor.MakeTextPath(K2Node_PropertyAccess->GetPath());	
 			const FText& CompilationContext = K2Node_PropertyAccess->GetCompiledContext();
 			const FText& CompilationContextDesc = K2Node_PropertyAccess->GetCompiledContextDesc();
 			if(CompilationContext.IsEmpty() && CompilationContextDesc.IsEmpty())
 			{
-				return TextPath;
+
+				return FText::Format(LOCTEXT("ToolTipFormat", "Access property '{0}'\nNative: {1}"), TextPath, UnderlyingPath);
 			}
 			else
 			{
-				return FText::Format(LOCTEXT("CompiledAccessToolTipFormat", "Access property '{0}'\n{1}\n{2}"), TextPath, CompilationContext, CompilationContextDesc);
+				return FText::Format(LOCTEXT("CompiledAccessToolTipFormat", "Access property '{0}'\nNative: {1}\n{2}\n{3}"), TextPath, UnderlyingPath, CompilationContext, CompilationContextDesc);
 			}
 		}
 	});
diff --git a/Engine/Plugins/Runtime/PropertyAccess/Source/PropertyAccessEditor/Private/PropertyAccessEditor.cpp b/Engine/Plugins/Runtime/PropertyAccess/Source/PropertyAccessEditor/Private/PropertyAccessEditor.cpp
--- a/Engine/Plugins/Runtime/PropertyAccess/Source/PropertyAccessEditor/Private/PropertyAccessEditor.cpp
+++ b/Engine/Plugins/Runtime/PropertyAccess/Source/PropertyAccessEditor/Private/PropertyAccessEditor.cpp
@@ -485,8 +485,8 @@ struct FPropertyAccessEditorSystem
 				FPropertyAccessCopyContext CopyContext;
 				CopyContext.Object = OutCopy.AssociatedObject;
 				CopyContext.ContextId = OutCopy.ContextId;
-				CopyContext.SourcePathAsText = PropertyAccess::MakeTextPath(OutCopy.SourcePath);
-				CopyContext.DestPathAsText = PropertyAccess::MakeTextPath(OutCopy.DestPath);
+				CopyContext.SourcePathAsText = PropertyAccess::MakeTextPath(OutCopy.SourcePath, InStruct);
+				CopyContext.DestPathAsText = PropertyAccess::MakeTextPath(OutCopy.DestPath, InStruct);
 				CopyContext.bSourceThreadSafe = DestContext.bWasThreadSafe;
 				CopyContext.bDestThreadSafe = SrcContext.bWasThreadSafe;
 				
@@ -551,7 +551,12 @@ namespace PropertyAccess
 		// Special case for object properties
 		if(InPropertyA->IsA<FObjectPropertyBase>() && InPropertyB->IsA<FObjectPropertyBase>())
 		{
-			return EPropertyAccessCompatibility::Compatible;
+			const FObjectPropertyBase* ObjectPropertyA = CastField<FObjectPropertyBase>(InPropertyA);
+			const FObjectPropertyBase* ObjectPropertyB = CastField<FObjectPropertyBase>(InPropertyB);
+			if(ObjectPropertyA->PropertyClass->IsChildOf(ObjectPropertyB->PropertyClass))
+			{
+				return EPropertyAccessCompatibility::Compatible;
+			}
 		}
 
 		// Extract underlying types for enums
@@ -689,18 +694,88 @@ namespace PropertyAccess
 		});
 	}
 
-	FText MakeTextPath(const TArray<FString>& InPath)
+	FText MakeTextPath(const TArray<FString>& InPath, const UStruct* InStruct = nullptr)
 	{
 		TStringBuilder<128> StringBuilder;
-
-		if(InPath.Num())
+		
+		if(InPath.Num() > 0)
 		{
-			StringBuilder.Append(InPath[0]);
+			const int32 LastIndex = InPath.Num() - 1;
+			bool bResolved = InStruct != nullptr;
+			
+			if(InStruct)
+			{
+				IPropertyAccessEditor::FResolvePropertyAccessArgs ResolveArgs;
+				auto PropertyFunction = [&StringBuilder](FProperty* InProperty, int32 InArrayIndex, bool bLast)
+				{
+					if(InProperty->IsNative())
+					{
+						if(const FString* ScriptNamePtr = InProperty->FindMetaData("ScriptName"))
+						{
+							StringBuilder.Append(*ScriptNamePtr);
+						}
+						else
+						{
+							StringBuilder.Append(InProperty->GetName());
+						}
+					}
+					else
+					{
+						StringBuilder.Append(InProperty->GetDisplayNameText().ToString());
+					}
+
+					if(InArrayIndex != INDEX_NONE)
+					{
+						StringBuilder.Appendf(TEXT("[%d]"), InArrayIndex);
+					}
+
+					if(!bLast)
+					{
+						StringBuilder.Append(TEXT("."));
+					}
+				};
+				
+				ResolveArgs.PropertyFunction = [&PropertyFunction, LastIndex](int32 InSegmentIndex, FProperty* InProperty, int32 InStaticArrayIndex)
+				{
+					PropertyFunction(InProperty, InStaticArrayIndex, InSegmentIndex == LastIndex);
+				};
+				ResolveArgs.ArrayFunction = [&PropertyFunction, LastIndex](int32 InSegmentIndex, FArrayProperty* InProperty, int32 InArrayIndex)
+				{
+					PropertyFunction(InProperty, InArrayIndex, InSegmentIndex == LastIndex);
+				};
+				ResolveArgs.FunctionFunction = [&StringBuilder, LastIndex](int32 InSegmentIndex, UFunction* InFunction, FProperty* /*ReturnProperty*/)
+				{
+					if(const FString* ScriptNamePtr = InFunction->FindMetaData("ScriptName"))
+					{
+						StringBuilder.Append(*ScriptNamePtr);
+					}
+					else
+					{
+						StringBuilder.Append(InFunction->GetName());
+					}
 
-			for(int32 SegmentIndex = 1; SegmentIndex < InPath.Num(); ++SegmentIndex)
+					if(InSegmentIndex != LastIndex)
+					{
+						StringBuilder.Append(TEXT("."));
+					}
+				};
+				
+				if(FPropertyAccessEditorSystem::ResolvePropertyAccess(InStruct, InPath, ResolveArgs).Result == EPropertyAccessResolveResult::Failed)
+				{
+					bResolved = false;
+				}
+			}
+
+			// Fallback to string concatenation if we didnt/couldnt resolve 
+			if(!bResolved)
 			{
-				StringBuilder.Append(TEXT("."));
-				StringBuilder.Append(InPath[SegmentIndex]);
+				StringBuilder.Append(InPath[0]);
+
+				for(int32 SegmentIndex = 1; SegmentIndex < InPath.Num(); ++SegmentIndex)
+				{
+					StringBuilder.Append(TEXT("."));
+					StringBuilder.Append(InPath[SegmentIndex]);
+				}
 			}
 		}
 
diff --git a/Engine/Plugins/Runtime/PropertyAccess/Source/PropertyAccessEditor/Private/PropertyAccessEditor.h b/Engine/Plugins/Runtime/PropertyAccess/Source/PropertyAccessEditor/Private/PropertyAccessEditor.h
--- a/Engine/Plugins/Runtime/PropertyAccess/Source/PropertyAccessEditor/Private/PropertyAccessEditor.h
+++ b/Engine/Plugins/Runtime/PropertyAccess/Source/PropertyAccessEditor/Private/PropertyAccessEditor.h
@@ -31,7 +31,7 @@ namespace PropertyAccess
 	extern void MakeStringPath(const TArray<FBindingChainElement>& InBindingChain, TArray<FString>& OutStringPath);
 
 	// Makes a text path from an array of path segments 
-	extern FText MakeTextPath(const TArray<FString>& InPath);
+	extern FText MakeTextPath(const TArray<FString>& InPath, const UStruct* InStruct);
 }
 
 // A helper structure used to compile a property access library
diff --git a/Engine/Plugins/Runtime/PropertyAccess/Source/PropertyAccessEditor/Private/PropertyAccessEditorModule.cpp b/Engine/Plugins/Runtime/PropertyAccess/Source/PropertyAccessEditor/Private/PropertyAccessEditorModule.cpp
--- a/Engine/Plugins/Runtime/PropertyAccess/Source/PropertyAccessEditor/Private/PropertyAccessEditorModule.cpp
+++ b/Engine/Plugins/Runtime/PropertyAccess/Source/PropertyAccessEditor/Private/PropertyAccessEditorModule.cpp
@@ -68,9 +68,9 @@ public:
 		return MakeUnique<FPropertyAccessLibraryCompiler>(&InArgs.Library, InArgs.ClassContext, InArgs.OnDetermineBatchId);
 	}
 
-	virtual FText MakeTextPath(const TArray<FString>& InPath) const override
+	virtual FText MakeTextPath(const TArray<FString>& InPath, const UStruct* InStruct = nullptr) const override
 	{
-		return PropertyAccess::MakeTextPath(InPath);
+		return PropertyAccess::MakeTextPath(InPath, InStruct);
 	};
 };
 
diff --git a/Engine/Plugins/Runtime/PropertyAccess/Source/PropertyAccessEditor/Private/SPropertyBinding.cpp b/Engine/Plugins/Runtime/PropertyAccess/Source/PropertyAccessEditor/Private/SPropertyBinding.cpp
--- a/Engine/Plugins/Runtime/PropertyAccess/Source/PropertyAccessEditor/Private/SPropertyBinding.cpp
+++ b/Engine/Plugins/Runtime/PropertyAccess/Source/PropertyAccessEditor/Private/SPropertyBinding.cpp
@@ -310,8 +310,35 @@ bool SPropertyBinding::HasBindablePropertiesRecursive(UStruct* InStruct, TSet<US
 			}
 		}
 	});
+
+	if(UClass* Class = Cast<UClass>(InStruct))
+	{
+		ForEachBindableFunction(Class, [this, &BindableCount, &VisitedStructs, RecursionDepth](const FFunctionInfo& Info)
+		{
+			FProperty* ReturnProperty = Info.Function->GetReturnProperty();
+			
+			// We can get here if we accept non-leaf UObject functions, so if so we need to check the return value for compatibility
+			if(!Args.bAllowUObjectFunctions || Args.OnCanBindProperty.Execute(ReturnProperty))
+			{
+				BindableCount++;
+			}
+
+			// Only show bindable subobjects, structs and variables if we're generating pure bindings.
+			if(Args.bGeneratePureBindings)
+			{
+				if(FObjectPropertyBase* ObjectPropertyBase = CastField<FObjectPropertyBase>(ReturnProperty))
+				{
+					HasBindablePropertiesRecursive(ObjectPropertyBase->PropertyClass, VisitedStructs, RecursionDepth + 1);
+				}
+				else if(FStructProperty* StructProperty = CastField<FStructProperty>(ReturnProperty))
+				{
+					HasBindablePropertiesRecursive(StructProperty->Struct, VisitedStructs, RecursionDepth + 1);
+				}
+			}
+		});
+	}
 	
-	return BindableCount > 0;				
+	return BindableCount > 0;
 }
 
 TSharedRef<SWidget> SPropertyBinding::OnGenerateDelegateMenu()
@@ -504,6 +531,23 @@ void SPropertyBinding::FillPropertyMenu(FMenuBuilder& MenuBuilder, UStruct* InOw
 		FEdGraphPinType PinType;
 		Schema->ConvertPropertyToPinType(InProperty, PinType);
 
+		FText PropertyDisplayName;
+		if(InProperty->IsNative())
+		{
+			if(const FString* ScriptNamePtr = InProperty->FindMetaData("ScriptName"))
+			{
+				PropertyDisplayName = FText::FromString(*ScriptNamePtr);
+			}
+			else
+			{
+				PropertyDisplayName = FText::FromString(InProperty->GetName());
+			}
+		}
+		else
+		{
+			PropertyDisplayName = FText::FromString(InProperty->GetDisplayNameText().ToString());
+		}
+		
 		return SNew(SHorizontalBox)
 			.ToolTipText(InProperty->GetToolTipText())
 			+SHorizontalBox::Slot()
@@ -527,7 +571,7 @@ void SPropertyBinding::FillPropertyMenu(FMenuBuilder& MenuBuilder, UStruct* InOw
 			.Padding(4.0f, 0.0f)
 			[
 				SNew(STextBlock)
-				.Text(FText::FromName(InProperty->GetFName()))
+				.Text(PropertyDisplayName)
 			];
 	};
 
@@ -607,15 +651,21 @@ void SPropertyBinding::FillPropertyMenu(FMenuBuilder& MenuBuilder, UStruct* InOw
 					{
 						if(FObjectPropertyBase* ObjectPropertyBase = CastField<FObjectPropertyBase>(ReturnProperty))
 						{
-							MenuBuilder.AddSubMenu(
-								MakeFunctionWidget(Info),
-								FNewMenuDelegate::CreateSP(this, &SPropertyBinding::FillPropertyMenu, static_cast<UStruct*>(ObjectPropertyBase->PropertyClass), NewBindingChain));
+							if (HasBindableProperties(ObjectPropertyBase->PropertyClass))
+							{
+								MenuBuilder.AddSubMenu(
+									MakeFunctionWidget(Info),
+									FNewMenuDelegate::CreateSP(this, &SPropertyBinding::FillPropertyMenu, static_cast<UStruct*>(ObjectPropertyBase->PropertyClass), NewBindingChain));
+							}
 						}
 						else if(FStructProperty* StructProperty = CastField<FStructProperty>(ReturnProperty))
 						{
-							MenuBuilder.AddSubMenu(
-								MakeFunctionWidget(Info),
-								FNewMenuDelegate::CreateSP(this, &SPropertyBinding::FillPropertyMenu, static_cast<UStruct*>(StructProperty->Struct), NewBindingChain));
+							if (HasBindableProperties(StructProperty->Struct))
+							{
+								MenuBuilder.AddSubMenu(
+									MakeFunctionWidget(Info),
+									FNewMenuDelegate::CreateSP(this, &SPropertyBinding::FillPropertyMenu, static_cast<UStruct*>(StructProperty->Struct), NewBindingChain));
+							}
 						}
 					}
 				});
diff --git a/Engine/Plugins/Runtime/PropertyAccess/Source/PropertyAccessEditor/Private/SPropertyBinding.h b/Engine/Plugins/Runtime/PropertyAccess/Source/PropertyAccessEditor/Private/SPropertyBinding.h
--- a/Engine/Plugins/Runtime/PropertyAccess/Source/PropertyAccessEditor/Private/SPropertyBinding.h
+++ b/Engine/Plugins/Runtime/PropertyAccess/Source/PropertyAccessEditor/Private/SPropertyBinding.h
@@ -37,7 +37,7 @@ protected:
 		}
 
 		FFunctionInfo(UFunction* InFunction)
-			: DisplayName(FText::FromName(InFunction->GetFName()))
+			: DisplayName(InFunction->HasMetaData("ScriptName") ? InFunction->GetMetaDataText("ScriptName") : FText::FromName(InFunction->GetFName()))
 			, Tooltip(InFunction->GetMetaData("Tooltip"))
 			, FuncName(InFunction->GetFName())
 			, Function(InFunction)
diff --git a/Engine/Source/Editor/AnimGraph/Private/AnimGraphNode_Base.cpp b/Engine/Source/Editor/AnimGraph/Private/AnimGraphNode_Base.cpp
--- a/Engine/Source/Editor/AnimGraph/Private/AnimGraphNode_Base.cpp
+++ b/Engine/Source/Editor/AnimGraph/Private/AnimGraphNode_Base.cpp
@@ -299,6 +299,16 @@ void UAnimGraphNode_Base::InternalPinCreation(TArray<UEdGraphPin*>* OldPins)
 		// Create the output pin, if needed
 		CreateOutputPins();
 	}
+
+	if(HasValidBlueprint())
+	{
+		// Update any binding's display text
+		IPropertyAccessEditor& PropertyAccessEditor = IModularFeatures::Get().GetModularFeature<IPropertyAccessEditor>("PropertyAccessEditor");
+		for(TPair<FName, FAnimGraphNodePropertyBinding>& BindingPair : PropertyBindings)
+		{
+			BindingPair.Value.PathAsText = PropertyAccessEditor.MakeTextPath(BindingPair.Value.PropertyPath, GetBlueprint()->SkeletonGeneratedClass);
+		}
+	}
 }
 
 void UAnimGraphNode_Base::AllocateDefaultPins()
@@ -1002,7 +1012,7 @@ TSharedRef<SWidget> UAnimGraphNode_Base::MakePropertyBindingWidget(const FAnimPr
 						Binding.ArrayIndex = InArgs.PinName.GetNumber() - 1;
 					}
 					PropertyAccessEditor.MakeStringPath(InBindingChain, Binding.PropertyPath);
-					Binding.PathAsText = PropertyAccessEditor.MakeTextPath(Binding.PropertyPath);
+					Binding.PathAsText = PropertyAccessEditor.MakeTextPath(Binding.PropertyPath, Blueprint->SkeletonGeneratedClass);
 					Binding.Type = LeafField.IsA<UFunction>() ? EAnimGraphNodePropertyBindingType::Function : EAnimGraphNodePropertyBindingType::Property;
 					Binding.bIsBound = true;
 					AnimGraphNode->RecalculateBindingType(Binding);
@@ -1159,15 +1169,17 @@ TSharedRef<SWidget> UAnimGraphNode_Base::MakePropertyBindingWidget(const FAnimPr
 					}
 					else
 					{
+						IPropertyAccessEditor& PropertyAccessEditor = IModularFeatures::Get().GetModularFeature<IPropertyAccessEditor>("PropertyAccessEditor");
+						const FText UnderlyingPath = PropertyAccessEditor.MakeTextPath(BindingPtr->PropertyPath);
 						const FText& CompilationContext = BindingPtr->CompiledContext;
 						const FText& CompilationContextDesc = BindingPtr->CompiledContextDesc;
 						if(CompilationContext.IsEmpty() && CompilationContextDesc.IsEmpty())
 						{
-							SetAssignValue(FText::Format(LOCTEXT("BindingToolTipFormat", "Pin is bound to property '{0}'"), BindingPtr->PathAsText), ECurrentValueType::Binding);
+							SetAssignValue(FText::Format(LOCTEXT("BindingToolTipFormat", "Pin is bound to property '{0}'\nNative: {1}"), BindingPtr->PathAsText, UnderlyingPath), ECurrentValueType::Binding);
 						}
 						else
 						{
-							SetAssignValue(FText::Format(LOCTEXT("BindingToolTipFormatWithDesc", "Pin is bound to property '{0}'\n{1}\n{2}"), BindingPtr->PathAsText, CompilationContext, CompilationContextDesc), ECurrentValueType::Binding);
+							SetAssignValue(FText::Format(LOCTEXT("BindingToolTipFormatWithDesc", "Pin is bound to property '{0}'\nNative: {1}\n{2}\n{2}"), BindingPtr->PathAsText, UnderlyingPath, CompilationContext, CompilationContextDesc), ECurrentValueType::Binding);
 						}
 					}
 				}
@@ -1644,7 +1656,7 @@ void UAnimGraphNode_Base::HandleVariableRenamed(UBlueprint* InBlueprint, UClass*
 		for(const int32& RenameIndex : RenameIndices)
 		{
 			BindingPair.Value.PropertyPath[RenameIndex] = InNewVarName.ToString();
-			BindingPair.Value.PathAsText = PropertyAccessEditor.MakeTextPath(BindingPair.Value.PropertyPath);
+			BindingPair.Value.PathAsText = PropertyAccessEditor.MakeTextPath(BindingPair.Value.PropertyPath, GetAnimBlueprint()->SkeletonGeneratedClass);
 		}
 	}
 }
@@ -1679,7 +1691,7 @@ void UAnimGraphNode_Base::ReplaceReferences(UBlueprint* InBlueprint, UBlueprint*
 		for(const int32& RenameIndex : ReplaceIndices)
 		{
 			BindingPair.Value.PropertyPath[RenameIndex] = ReplacementProperty->GetName();
-			BindingPair.Value.PathAsText = PropertyAccessEditor.MakeTextPath(BindingPair.Value.PropertyPath);
+			BindingPair.Value.PathAsText = PropertyAccessEditor.MakeTextPath(BindingPair.Value.PropertyPath, GetAnimBlueprint()->SkeletonGeneratedClass);
 		}
 	}
 }
diff --git a/Engine/Source/Editor/UnrealEd/Public/IPropertyAccessEditor.h b/Engine/Source/Editor/UnrealEd/Public/IPropertyAccessEditor.h
--- a/Engine/Source/Editor/UnrealEd/Public/IPropertyAccessEditor.h
+++ b/Engine/Source/Editor/UnrealEd/Public/IPropertyAccessEditor.h
@@ -362,6 +362,8 @@ public:
 	// Make a property access library compiler, used for building a FPropertyAccessLibrary
 	virtual TUniquePtr<IPropertyAccessLibraryCompiler> MakePropertyAccessCompiler(const FPropertyAccessLibraryCompilerArgs& InArgs) const = 0;
 
-	// Make a text representation of a property path 
-	virtual FText MakeTextPath(const TArray<FString>& InPath) const = 0;
+	// Make a text representation of a property path
+	// @param	InPath		The path to use to generate the text path
+	// @param	InStruct	Optional struct to resolve against - if this is supplied then the text path can use 'correct' display names
+	virtual FText MakeTextPath(const TArray<FString>& InPath, const UStruct* InStruct = nullptr) const = 0;
 };
\ No newline at end of file
