Property Access: Allow to bind to context structs directly.
- Added ability to return a context struct at binding
- Added categories to context structs, which are shown on the menu
- Use BP pin type as default icon for context structs

diff --git a/Engine/Plugins/Runtime/PropertyAccess/Source/PropertyAccessEditor/Private/SPropertyBinding.cpp b/Engine/Plugins/Runtime/PropertyAccess/Source/PropertyAccessEditor/Private/SPropertyBinding.cpp
--- a/Engine/Plugins/Runtime/PropertyAccess/Source/PropertyAccessEditor/Private/SPropertyBinding.cpp
+++ b/Engine/Plugins/Runtime/PropertyAccess/Source/PropertyAccessEditor/Private/SPropertyBinding.cpp
@@ -394,9 +394,33 @@ TSharedRef<SWidget> SPropertyBinding::OnGenerateDelegateMenu()
 	{
 		auto MakeContextStructWidget = [this](const FBindingContextStruct& ContextStruct)
 		{
-			FText DisplayText = ContextStruct.DisplayText.IsEmpty() ? ContextStruct.Struct->GetDisplayNameText() : ContextStruct.DisplayText;
-			FText ToolTipText = ContextStruct.TooltipText.IsEmpty() ? DisplayText : ContextStruct.TooltipText;
+			const FText DisplayText = ContextStruct.DisplayText.IsEmpty() ? ContextStruct.Struct->GetDisplayNameText() : ContextStruct.DisplayText;
+			const FText ToolTipText = ContextStruct.TooltipText.IsEmpty() ? DisplayText : ContextStruct.TooltipText;
 
+			const FSlateBrush* Icon = ContextStruct.Icon;
+			FLinearColor IconColor = FLinearColor::White;
+			if (Icon == nullptr)
+			{
+				const UEdGraphSchema_K2* Schema = GetDefault<UEdGraphSchema_K2>();
+
+				FEdGraphPinType PinType;
+
+				if (UClass* Class = Cast<UClass>(ContextStruct.Struct))
+				{
+					PinType.PinCategory = UEdGraphSchema_K2::PC_Object;
+					PinType.PinSubCategory = NAME_None;
+					PinType.PinSubCategoryObject = Class;
+				}
+				else if (UScriptStruct* ScriptStruct = Cast<UScriptStruct>(ContextStruct.Struct))
+				{
+					PinType.PinCategory = UEdGraphSchema_K2::PC_Struct;
+					PinType.PinSubCategory = NAME_None;
+					PinType.PinSubCategoryObject = ScriptStruct;
+				}
+				Icon = FBlueprintEditorUtils::GetIconFromPin(PinType, true); 
+				IconColor = Schema->GetPinTypeColor(PinType);
+			}
+			
 			return SNew(SHorizontalBox)
 				.ToolTipText(ToolTipText)
 				+ SHorizontalBox::Slot()
@@ -411,7 +435,8 @@ TSharedRef<SWidget> SPropertyBinding::OnGenerateDelegateMenu()
 				.Padding(1.0f, 0.0f)
 				[
 					SNew(SImage)
-					.Image(ContextStruct.Icon)
+					.Image(Icon)
+					.ColorAndOpacity(IconColor)
 				]
 				+ SHorizontalBox::Slot()
 				.AutoWidth()
@@ -423,20 +448,50 @@ TSharedRef<SWidget> SPropertyBinding::OnGenerateDelegateMenu()
 				];
 		};
 
+		bool bSectionOpened = false;
+		FText CurrentSection = FText::GetEmpty();
+		
 		for (int32 i = 0; i < BindingContextStructs.Num(); i++)
 		{
 			FBindingContextStruct& ContextStruct = BindingContextStructs[i];
-			if (HasBindableProperties(ContextStruct.Struct))
+
+			if (!ContextStruct.Section.IdenticalTo(CurrentSection))
 			{
-				// Make first chain element representing the index in the context array.
-				TArray<TSharedPtr<FBindingChainElement>> BindingChain;
-				BindingChain.Emplace(MakeShared<FBindingChainElement>(nullptr, i));
+				if (bSectionOpened)
+				{
+					MenuBuilder.EndSection();
+				}
 
+				CurrentSection = ContextStruct.Section;
+				MenuBuilder.BeginSection(FName(), CurrentSection);
+				bSectionOpened = true;
+			}
+			
+			// Make first chain element representing the index in the context array.
+			TArray<TSharedPtr<FBindingChainElement>> BindingChain;
+			BindingChain.Emplace(MakeShared<FBindingChainElement>(nullptr, i));
+
+			if (Args.OnCanBindToContextStruct.IsBound() && Args.OnCanBindToContextStruct.Execute(ContextStruct.Struct))
+			{
+				// If the struct can be be bound to directly, create action for that. 
+				MenuBuilder.AddMenuEntry(
+					FUIAction(FExecuteAction::CreateSP(this, &SPropertyBinding::HandleAddBinding, BindingChain)),
+					MakeContextStructWidget(ContextStruct));
+			}
+			else if (HasBindableProperties(ContextStruct.Struct))
+			{
+				// Show struct properties.
 				MenuBuilder.AddSubMenu(
 					MakeContextStructWidget(ContextStruct),
 					FNewMenuDelegate::CreateSP(this, &SPropertyBinding::FillPropertyMenu, ContextStruct.Struct, BindingChain));
 			}
 		}
+
+		if (bSectionOpened)
+		{
+			MenuBuilder.EndSection();
+		}
+
 	}
 
 	FDisplayMetrics DisplayMetrics;
@@ -524,8 +579,6 @@ void SPropertyBinding::FillPropertyMenu(FMenuBuilder& MenuBuilder, UStruct* InOw
 
 	auto MakePropertyWidget = [this](FProperty* InProperty)
 	{
-		static FName PropertyIcon(TEXT("Kismet.VariableList.TypeIcon"));
-
 		const UEdGraphSchema_K2* Schema = GetDefault<UEdGraphSchema_K2>();
 
 		FEdGraphPinType PinType;
diff --git a/Engine/Source/Editor/UnrealEd/Public/IPropertyAccessEditor.h b/Engine/Source/Editor/UnrealEd/Public/IPropertyAccessEditor.h
--- a/Engine/Source/Editor/UnrealEd/Public/IPropertyAccessEditor.h
+++ b/Engine/Source/Editor/UnrealEd/Public/IPropertyAccessEditor.h
@@ -67,24 +67,30 @@ struct FRedirectorBindingInfo
  */
 struct FBindingContextStruct
 {
-	FBindingContextStruct(UStruct* InStruct, const FSlateBrush* InIcon = nullptr, const FText& InDisplayText = FText::GetEmpty(), const FText& InTooltipText = FText::GetEmpty())
+	FBindingContextStruct() = default;
+	
+	FBindingContextStruct(UStruct* InStruct, const FSlateBrush* InIcon = nullptr, const FText& InDisplayText = FText::GetEmpty(), const FText& InTooltipText = FText::GetEmpty(), const FText& InSection = FText::GetEmpty())
 		: Struct(InStruct)
 		, Icon(InIcon)
 		, DisplayText(InDisplayText)
 		, TooltipText(InTooltipText)
+		, Section(InSection)
 	{}
 
 	/** The struct to bind to. */
-	UStruct* Struct;
+	UStruct* Struct = nullptr;
 
-	/** Icon to display in hte poppup menu. */ 
-	const FSlateBrush* Icon;
+	/** Icon to display in the popup menu. */ 
+	const FSlateBrush* Icon = nullptr;
 
 	/** Text to display for this item in the popup. If left empty, struct's display text will be used. */
 	FText DisplayText;
 
 	/** Tooltip to show, or if empty, the tool tip will be set to the same as popup text. */
 	FText TooltipText;
+
+	/** Name of the section to put the struct to. If left empty, no section will be created. */
+	FText Section;
 };
 
 
@@ -109,6 +115,9 @@ DECLARE_DELEGATE_RetVal_OneParam(bool, FOnCanBindFunction, UFunction* /*InFuncti
 /** Delegate called to see if a class can be bound to */
 DECLARE_DELEGATE_RetVal_OneParam(bool, FOnCanBindToClass, UClass* /*InClass*/);
 
+/** Delegate called to see if a class can be bound to */
+DECLARE_DELEGATE_RetVal_OneParam(bool, FOnCanBindToContextStruct, UStruct* /*InStruct*/);
+
 /** Delegate called to see if a subobject can be bound to */
 DECLARE_DELEGATE_RetVal_OneParam(bool, FOnCanBindToSubObjectClass, UClass* /*InSubObjectClass*/);
 
@@ -154,6 +163,9 @@ struct FPropertyBindingWidgetArgs
 	/** Delegate called to see if a class can be bound to */
 	FOnCanBindToClass OnCanBindToClass;
 
+	/** Delegate called to see if a context struct can be directly bound to */
+	FOnCanBindToContextStruct OnCanBindToContextStruct;
+	
 	/** Delegate called to see if a subobject can be bound to */
 	FOnCanBindToSubObjectClass OnCanBindToSubObjectClass;
 
