Add support for LWC promotion when copying TMap properties via property access

TMaps with double values were being reported as incompatible with TMaps with float values

diff --git a/Engine/Plugins/Runtime/PropertyAccess/Source/PropertyAccessEditor/Private/PropertyAccessEditor.cpp b/Engine/Plugins/Runtime/PropertyAccess/Source/PropertyAccessEditor/Private/PropertyAccessEditor.cpp
--- a/Engine/Plugins/Runtime/PropertyAccess/Source/PropertyAccessEditor/Private/PropertyAccessEditor.cpp
+++ b/Engine/Plugins/Runtime/PropertyAccess/Source/PropertyAccessEditor/Private/PropertyAccessEditor.cpp
@@ -476,6 +476,26 @@ struct FPropertyAccessEditorSystem
 					}
 				}
 			}
+			else if (SrcProperty->IsA<FMapProperty>() && DestProperty->IsA<FMapProperty>())
+			{
+				const FMapProperty* SrcMapProperty = CastField<const FMapProperty>(SrcProperty);
+				const FMapProperty* DestMapProperty = CastField<const FMapProperty>(DestProperty);
+
+				if (SrcMapProperty->ValueProp->IsA<FFloatProperty>())
+				{
+					if (DestMapProperty->ValueProp->IsA<FDoubleProperty>())
+					{
+						return EPropertyAccessCopyType::PromoteMapValueFloatToDouble;
+					}
+				}
+				else if (SrcMapProperty->ValueProp->IsA<FDoubleProperty>())
+				{
+					if (DestMapProperty->ValueProp->IsA<FFloatProperty>())
+					{
+						return EPropertyAccessCopyType::DemoteMapValueDoubleToFloat;
+					}
+				}
+			}
 		}
 
 		OutErrorMessage = FText::Format(LOCTEXT("CopyTypeInvalidFormat", "@@ Cannot copy property ({0} -> {1})"), FText::FromString(SrcProperty->GetCPPType()), FText::FromString(DestProperty->GetCPPType()));
@@ -652,6 +672,47 @@ namespace PropertyAccess
 					}
 				}
 			}
+			else if (InPropertyA->IsA<FSetProperty>() && InPropertyB->IsA<FSetProperty>())
+			{
+				const FSetProperty* SetPropertyA = CastField<const FSetProperty>(InPropertyA);
+				const FSetProperty* SetPropertyB = CastField<const FSetProperty>(InPropertyB);
+
+				if (SetPropertyA->ElementProp->IsA<FFloatProperty>())
+				{
+					if (SetPropertyB->ElementProp->IsA<FDoubleProperty>())
+					{
+						return EPropertyAccessCompatibility::Promotable;
+					}
+				}
+				else if (SetPropertyA->ElementProp->IsA<FDoubleProperty>())
+				{
+					if (SetPropertyB->ElementProp->IsA<FFloatProperty>())
+					{
+						return EPropertyAccessCompatibility::Promotable;
+					}
+				}
+			}
+			else if (InPropertyA->IsA<FMapProperty>() && InPropertyB->IsA<FMapProperty>())
+			{
+				// We only support promoting value properties of maps. Is there a case where float/double keys are useful?
+				const FMapProperty* MapPropertyA = CastField<const FMapProperty>(InPropertyA);
+				const FMapProperty* MapPropertyB = CastField<const FMapProperty>(InPropertyB);
+
+				if (MapPropertyA->ValueProp->IsA<FFloatProperty>())
+				{
+					if (MapPropertyB->ValueProp->IsA<FDoubleProperty>())
+					{
+						return EPropertyAccessCompatibility::Promotable;
+					}
+				}
+				else if (MapPropertyA->ValueProp->IsA<FDoubleProperty>())
+				{
+					if (MapPropertyB->ValueProp->IsA<FFloatProperty>())
+					{
+						return EPropertyAccessCompatibility::Promotable;
+					}
+				}
+			}
 		}
 
 		return EPropertyAccessCompatibility::Incompatible;
diff --git a/Engine/Source/Runtime/Engine/Private/PropertyAccess.cpp b/Engine/Source/Runtime/Engine/Private/PropertyAccess.cpp
--- a/Engine/Source/Runtime/Engine/Private/PropertyAccess.cpp
+++ b/Engine/Source/Runtime/Engine/Private/PropertyAccess.cpp
@@ -396,6 +396,24 @@ struct FPropertyAccessSystem
 			CopyAndCastFloatingPointArray<double, float>(SrcArrayProperty, InSrcAddr, DestArrayProperty, InDestAddr);
 			break;
 		}
+		case EPropertyAccessCopyType::PromoteMapValueFloatToDouble:
+			{
+				checkSlow(InSrcProperty->IsA<FMapProperty>());
+				checkSlow(InDestProperty->IsA<FMapProperty>());
+				const FMapProperty* SrcMapProperty = ExactCastField<const FMapProperty>(InSrcProperty);
+				const FMapProperty* DestMapProperty = ExactCastField<const FMapProperty>(InDestProperty);
+				CopyAndCastFloatingPointMapValues<float, double>(SrcMapProperty, InSrcAddr, DestMapProperty, InDestAddr);
+				break;
+			}
+		case EPropertyAccessCopyType::DemoteMapValueDoubleToFloat:
+			{
+				checkSlow(InSrcProperty->IsA<FMapProperty>());
+				checkSlow(InDestProperty->IsA<FMapProperty>());
+				const FMapProperty* SrcMapProperty = ExactCastField<const FMapProperty>(InSrcProperty);
+				const FMapProperty* DestMapProperty = ExactCastField<const FMapProperty>(InDestProperty);
+				CopyAndCastFloatingPointMapValues<double, float>(SrcMapProperty, InSrcAddr, DestMapProperty, InDestAddr);
+				break;
+			}	
 		default:
 			check(false);
 			break;
@@ -629,6 +647,30 @@ struct FPropertyAccessSystem
 			*DestinationData = static_cast<DestinationType>(*SourceData);
 		}
 	}
+
+	template <typename SourceType, typename DestinationType>
+	static void CopyAndCastFloatingPointMapValues(const FMapProperty* SourceMapProperty,
+												  const void* SourceAddress,
+												  const FMapProperty* DestinationMapProperty,
+												  void* DestinationAddress)
+	{
+		checkSlow(SourceMapProperty);
+		checkSlow(SourceAddress);
+		checkSlow(DestinationMapProperty);
+		checkSlow(DestinationAddress);
+
+		FScriptMapHelper SourceMapHelper(SourceMapProperty, SourceAddress);
+		FScriptMapHelper DestinationMapHelper(DestinationMapProperty, DestinationAddress);
+
+		DestinationMapHelper.EmptyValues();
+		for (int32 i = 0; i < SourceMapHelper.Num(); ++i)
+		{
+			const void* KeyData = SourceMapHelper.GetKeyPtr(i);
+			const SourceType* SourceValueData = reinterpret_cast<const SourceType*>(SourceMapHelper.GetValuePtr(i));
+			DestinationType CastedType = static_cast<DestinationType>(*SourceValueData);
+			DestinationMapHelper.AddPair(KeyData, &CastedType);
+		}
+	}	
 };
 
 const FPropertyAccessLibrary& FPropertyAccessLibrary::operator =(const FPropertyAccessLibrary& Other)
diff --git a/Engine/Source/Runtime/Engine/Public/PropertyAccess.h b/Engine/Source/Runtime/Engine/Public/PropertyAccess.h
--- a/Engine/Source/Runtime/Engine/Public/PropertyAccess.h
+++ b/Engine/Source/Runtime/Engine/Public/PropertyAccess.h
@@ -345,6 +345,9 @@ enum class EPropertyAccessCopyType : uint8
 
 	PromoteArrayFloatToDouble,
 	DemoteArrayDoubleToFloat,
+
+	PromoteMapValueFloatToDouble,
+	DemoteMapValueDoubleToFloat,
 };
 
 // A property copy, represents a one-to-many copy operation
