Property Access editor: added support to set binding button background.

diff --git a/Engine/Plugins/Runtime/PropertyAccess/Source/PropertyAccessEditor/Private/SPropertyBinding.cpp b/Engine/Plugins/Runtime/PropertyAccess/Source/PropertyAccessEditor/Private/SPropertyBinding.cpp
--- a/Engine/Plugins/Runtime/PropertyAccess/Source/PropertyAccessEditor/Private/SPropertyBinding.cpp
+++ b/Engine/Plugins/Runtime/PropertyAccess/Source/PropertyAccessEditor/Private/SPropertyBinding.cpp
@@ -39,6 +39,7 @@ void SPropertyBinding::Construct(const FArguments& InArgs, UBlueprint* InBluepri
 		.FillWidth(1.0f)
 		[
 			SNew(SComboButton)
+			.ButtonStyle(Args.BindButtonStyle)
 			.ToolTipText(this, &SPropertyBinding::GetCurrentBindingToolTipText)
 			.OnGetMenuContent(this, &SPropertyBinding::OnGenerateDelegateMenu)
 			.ContentPadding(1)
@@ -65,6 +66,7 @@ void SPropertyBinding::Construct(const FArguments& InArgs, UBlueprint* InBluepri
 				[
 					SNew(STextBlock)
 					.Text(this, &SPropertyBinding::GetCurrentBindingText)
+					.ColorAndOpacity(FSlateColor::UseForeground())
 					.Font(IDetailLayoutBuilder::GetDetailFont())
 				]
 			]
diff --git a/Engine/Source/Editor/UnrealEd/Public/IPropertyAccessEditor.h b/Engine/Source/Editor/UnrealEd/Public/IPropertyAccessEditor.h
--- a/Engine/Source/Editor/UnrealEd/Public/IPropertyAccessEditor.h
+++ b/Engine/Source/Editor/UnrealEd/Public/IPropertyAccessEditor.h
@@ -196,6 +196,9 @@ struct FPropertyBindingWidgetArgs
 	/** Menu extender */
 	TSharedPtr<FExtender> MenuExtender;
 
+	/** Optional style override for bind button */
+	const FButtonStyle* BindButtonStyle = nullptr;
+	
 	/** Whether to generate pure bindings */
 	bool bGeneratePureBindings = true;
 
