Fix nullptr dereference static analysis warning in PropertyAccessEditor

diff --git a/Engine/Plugins/Runtime/PropertyAccess/Source/PropertyAccessEditor/Private/PropertyAccessEditor.cpp b/Engine/Plugins/Runtime/PropertyAccess/Source/PropertyAccessEditor/Private/PropertyAccessEditor.cpp
--- a/Engine/Plugins/Runtime/PropertyAccess/Source/PropertyAccessEditor/Private/PropertyAccessEditor.cpp
+++ b/Engine/Plugins/Runtime/PropertyAccess/Source/PropertyAccessEditor/Private/PropertyAccessEditor.cpp
@@ -149,6 +149,8 @@ struct FPropertyAccessEditorSystem
 
 	static ESegmentResolveResult ResolveSegments_CheckFunction(FPropertyAccessSegment& InSegment, UFunction* InFunction, FResolveSegmentsContext& InContext)
 	{
+		check(InFunction);
+		
 		InSegment.Function = InFunction;
 		InSegment.Flags |= (uint16)EPropertyAccessSegmentFlags::Function;
 
