Fix for static analysis in PropertyAccessEditor

diff --git a/Engine/Plugins/Runtime/PropertyAccess/Source/PropertyAccessEditor/Private/PropertyAccessEditor.cpp b/Engine/Plugins/Runtime/PropertyAccess/Source/PropertyAccessEditor/Private/PropertyAccessEditor.cpp
--- a/Engine/Plugins/Runtime/PropertyAccess/Source/PropertyAccessEditor/Private/PropertyAccessEditor.cpp
+++ b/Engine/Plugins/Runtime/PropertyAccess/Source/PropertyAccessEditor/Private/PropertyAccessEditor.cpp
@@ -150,11 +150,11 @@ struct FPropertyAccessEditorSystem
 	static ESegmentResolveResult ResolveSegments_CheckFunction(FPropertyAccessSegment& InSegment, UFunction* InFunction, FResolveSegmentsContext& InContext)
 	{
 		check(InFunction);
-		
+
 		InSegment.Function = InFunction;
 		InSegment.Flags |= (uint16)EPropertyAccessSegmentFlags::Function;
 
-		if (InContext.bPerformValidation && InFunction && !FBlueprintActionDatabase::IsFunctionAllowed(InFunction, FBlueprintActionDatabase::EPermissionsContext::Property))
+		if (InContext.bPerformValidation && !FBlueprintActionDatabase::IsFunctionAllowed(InFunction, FBlueprintActionDatabase::EPermissionsContext::Property))
 		{
 			InContext.ErrorMessage = FText::Format(LOCTEXT("FunctionIsDisallowedByPermissions", "Function '{0}' is not allowed for property access for @@"), FText::FromName(InFunction->GetFName()));
 			return ESegmentResolveResult::Failed;
