StateTree: Added handling bindings from instanced sources.

diff --git a/Engine/Plugins/Runtime/PropertyAccess/Source/PropertyAccessEditor/Private/SPropertyBinding.cpp b/Engine/Plugins/Runtime/PropertyAccess/Source/PropertyAccessEditor/Private/SPropertyBinding.cpp
--- a/Engine/Plugins/Runtime/PropertyAccess/Source/PropertyAccessEditor/Private/SPropertyBinding.cpp
+++ b/Engine/Plugins/Runtime/PropertyAccess/Source/PropertyAccessEditor/Private/SPropertyBinding.cpp
@@ -245,13 +245,13 @@ void SPropertyBinding::ForEachBindableProperty(UStruct* InStruct, Predicate Pred
 	}
 }
 
-bool SPropertyBinding::HasBindableProperties(UStruct* InStruct) const
+bool SPropertyBinding::HasBindableProperties(UStruct* InStruct, TArray<TSharedPtr<FBindingChainElement>>& BindingChain) const
 {
 	TSet<UStruct*> VisitedStructs;
-	return HasBindablePropertiesRecursive(InStruct, VisitedStructs, 0);
+	return HasBindablePropertiesRecursive(InStruct, VisitedStructs, BindingChain);
 }
 
-bool SPropertyBinding::HasBindablePropertiesRecursive(UStruct* InStruct, TSet<UStruct*>& VisitedStructs, const int32 RecursionDepth) const
+bool SPropertyBinding::HasBindablePropertiesRecursive(UStruct* InStruct, TSet<UStruct*>& VisitedStructs, TArray<TSharedPtr<FBindingChainElement>>& BindingChain) const
 {
 	if (VisitedStructs.Contains(InStruct))
 	{
@@ -262,14 +262,17 @@ bool SPropertyBinding::HasBindablePropertiesRecursive(UStruct* InStruct, TSet<US
 	VisitedStructs.Add(InStruct);
 	
 	// Arbitrary cut off to avoid infinite loops.
-	if (RecursionDepth > 10)
+	if (BindingChain.Num() > 10)
 	{
 		return false;
 	}
 	
 	int32 BindableCount = 0;
-	ForEachBindableProperty(InStruct, [this, &BindableCount, &VisitedStructs, RecursionDepth] (FProperty* Property)
+	ForEachBindableProperty(InStruct, [this, &BindableCount, &VisitedStructs, &BindingChain] (FProperty* Property)
 	{
+		BindingChain.Emplace(MakeShared<FBindingChainElement>(Property));
+		ON_SCOPE_EXIT{ BindingChain.Pop(); };
+
 		FArrayProperty* ArrayProperty = CastField<FArrayProperty>(Property);
 
 		if(Args.OnCanBindProperty.Execute(Property))
@@ -305,34 +308,45 @@ bool SPropertyBinding::HasBindablePropertiesRecursive(UStruct* InStruct, TSet<US
 			Struct = StructProperty->Struct;
 		}
 
-		// Recurse into a struct, except if it is the same type as the one we're binding.
-		if (Struct && HasBindablePropertiesRecursive(Struct, VisitedStructs, RecursionDepth + 1))
+		if (Struct)
 		{
-			if (Class)
+			if (UStruct* ResolvedStruct = ResolveIndirection(BindingChain))
 			{
-				// Ignore any subobject properties that are not bindable.
-				// Also ignore any class that is explicitly on the deny list.
-				if ( IsClassDenied(Class) || (Args.OnCanBindToSubObjectClass.IsBound() && Args.OnCanBindToSubObjectClass.Execute(Class)))
-				{
-					return;
-				}
+				Struct = ResolvedStruct;
 			}
 
-			if (Args.bAllowArrayElementBindings && ArrayProperty != nullptr)
-			{
-				BindableCount++;
-			}
-			else if (Args.bAllowStructMemberBindings)
+			// Recurse into a struct, except if it is the same type as the one we're binding.
+			if (HasBindablePropertiesRecursive(Struct, VisitedStructs, BindingChain))
 			{
-				BindableCount++;
+				if (Class)
+				{
+					// Ignore any subobject properties that are not bindable.
+					// Also ignore any class that is explicitly on the deny list.
+					if (IsClassDenied(Class) || (Args.OnCanBindToSubObjectClass.IsBound() && Args.OnCanBindToSubObjectClass.Execute(Class)))
+					{
+						return;
+					}
+				}
+
+				if (Args.bAllowArrayElementBindings && ArrayProperty != nullptr)
+				{
+					BindableCount++;
+				}
+				else if (Args.bAllowStructMemberBindings)
+				{
+					BindableCount++;
+				}
 			}
 		}
 	});
 
 	if(UClass* Class = Cast<UClass>(InStruct))
 	{
-		ForEachBindableFunction(Class, [this, &BindableCount, &VisitedStructs, RecursionDepth](const FFunctionInfo& Info)
+		ForEachBindableFunction(Class, [this, &BindableCount, &VisitedStructs, &BindingChain](const FFunctionInfo& Info)
 		{
+			BindingChain.Emplace(MakeShared<FBindingChainElement>(Info.Function));
+			ON_SCOPE_EXIT{ BindingChain.Pop(); };
+
 			FProperty* ReturnProperty = Info.Function->GetReturnProperty();
 			
 			// We can get here if we accept non-leaf UObject functions, so if so we need to check the return value for compatibility
@@ -346,11 +360,11 @@ bool SPropertyBinding::HasBindablePropertiesRecursive(UStruct* InStruct, TSet<US
 			{
 				if(FObjectPropertyBase* ObjectPropertyBase = CastField<FObjectPropertyBase>(ReturnProperty))
 				{
-					HasBindablePropertiesRecursive(ObjectPropertyBase->PropertyClass, VisitedStructs, RecursionDepth + 1);
+					HasBindablePropertiesRecursive(ObjectPropertyBase->PropertyClass, VisitedStructs, BindingChain);
 				}
 				else if(FStructProperty* StructProperty = CastField<FStructProperty>(ReturnProperty))
 				{
-					HasBindablePropertiesRecursive(StructProperty->Struct, VisitedStructs, RecursionDepth + 1);
+					HasBindablePropertiesRecursive(StructProperty->Struct, VisitedStructs, BindingChain);
 				}
 			}
 		});
@@ -496,7 +510,7 @@ TSharedRef<SWidget> SPropertyBinding::OnGenerateDelegateMenu()
 					FUIAction(FExecuteAction::CreateSP(this, &SPropertyBinding::HandleAddBinding, BindingChain)),
 					MakeContextStructWidget(ContextStruct));
 			}
-			if (HasBindableProperties(ContextStruct.Struct))
+			if (HasBindableProperties(ContextStruct.Struct, BindingChain))
 			{
 				// Show struct properties.
 				MenuBuilder.AddSubMenu(
@@ -722,7 +736,7 @@ void SPropertyBinding::FillPropertyMenu(FMenuBuilder& MenuBuilder, UStruct* InOw
 					{
 						if(FObjectPropertyBase* ObjectPropertyBase = CastField<FObjectPropertyBase>(ReturnProperty))
 						{
-							if (HasBindableProperties(ObjectPropertyBase->PropertyClass))
+							if (HasBindableProperties(ObjectPropertyBase->PropertyClass, NewBindingChain))
 							{
 								MenuBuilder.AddSubMenu(
 									MakeFunctionWidget(Info),
@@ -731,7 +745,7 @@ void SPropertyBinding::FillPropertyMenu(FMenuBuilder& MenuBuilder, UStruct* InOw
 						}
 						else if(FStructProperty* StructProperty = CastField<FStructProperty>(ReturnProperty))
 						{
-							if (HasBindableProperties(StructProperty->Struct))
+							if (HasBindableProperties(StructProperty->Struct, NewBindingChain))
 							{
 								MenuBuilder.AddSubMenu(
 									MakeFunctionWidget(Info),
@@ -824,6 +838,11 @@ void SPropertyBinding::FillPropertyMenu(FMenuBuilder& MenuBuilder, UStruct* InOw
 					// Recurse into a struct if it has some properties we can bind to.
 					if (Struct)
 					{
+						if (UStruct* ResolvedStruct = ResolveIndirection(NewBindingChain))
+						{
+							Struct = ResolvedStruct;
+						}
+
 						if (Class)
 						{
 							// Ignore any subobject properties that are not bindable.
@@ -834,11 +853,11 @@ void SPropertyBinding::FillPropertyMenu(FMenuBuilder& MenuBuilder, UStruct* InOw
 							}
 						}
 
-						if (HasBindableProperties(Struct))
+						if (HasBindableProperties(Struct, NewBindingChain))
 						{
 							if (Args.bAllowArrayElementBindings && ArrayProperty != nullptr)
 							{
-								TArray<TSharedPtr<FBindingChainElement>> NewArrayElementBindingChain(InBindingChain);
+								TArray<TSharedPtr<FBindingChainElement>> NewArrayElementBindingChain(NewBindingChain);
 								NewArrayElementBindingChain.Emplace(MakeShared<FBindingChainElement>(Property));
 								NewArrayElementBindingChain.Last()->ArrayIndex = 0;
 
@@ -1011,6 +1030,24 @@ void SPropertyBinding::HandleCreateAndAddBinding()
 	HandleGotoBindingClicked();
 }
 
+UStruct* SPropertyBinding::ResolveIndirection(const TArray<TSharedPtr<FBindingChainElement>>& BindingChain) const
+{
+	UStruct* ResolvedStruct = nullptr;
+
+	if (Args.OnResolveIndirection.IsBound())
+	{
+		TArray<FBindingChainElement> RawBindingChain;
+		Algo::Transform(BindingChain, RawBindingChain, [](TSharedPtr<FBindingChainElement> InElement)
+		{
+			return *InElement.Get();
+		});
+
+		ResolvedStruct = Args.OnResolveIndirection.Execute(RawBindingChain);
+	}
+
+	return ResolvedStruct;
+}
+
 EVisibility SPropertyBinding::GetGotoBindingVisibility() const
 {
 	if(Args.OnCanGotoBinding.IsBound())
diff --git a/Engine/Plugins/Runtime/PropertyAccess/Source/PropertyAccessEditor/Private/SPropertyBinding.h b/Engine/Plugins/Runtime/PropertyAccess/Source/PropertyAccessEditor/Private/SPropertyBinding.h
--- a/Engine/Plugins/Runtime/PropertyAccess/Source/PropertyAccessEditor/Private/SPropertyBinding.h
+++ b/Engine/Plugins/Runtime/PropertyAccess/Source/PropertyAccessEditor/Private/SPropertyBinding.h
@@ -62,6 +62,8 @@ protected:
 
 	void HandleCreateAndAddBinding();
 
+	UStruct* ResolveIndirection(const TArray<TSharedPtr<FBindingChainElement>>& BindingChain) const;
+
 	EVisibility GetGotoBindingVisibility() const;
 
 	FReply HandleGotoBindingClicked();
@@ -71,8 +73,8 @@ protected:
 private:
 	bool IsClassDenied(UClass* OwnerClass) const;
 	bool IsFieldFromDeniedClass(FFieldVariant Field) const;
-	bool HasBindableProperties(UStruct* InStruct) const;
-	bool HasBindablePropertiesRecursive(UStruct* InStruct, TSet<UStruct*>& VisitedStructs, const int32 RecursionDepth) const;
+	bool HasBindableProperties(UStruct* InStruct, TArray<TSharedPtr<FBindingChainElement>>& BindingChain) const;
+	bool HasBindablePropertiesRecursive(UStruct* InStruct, TSet<UStruct*>& VisitedStructs, TArray<TSharedPtr<FBindingChainElement>>& BindingChain) const;
 	
 	template <typename Predicate>
 	void ForEachBindableProperty(UStruct* InStruct, Predicate Pred) const;
diff --git a/Engine/Plugins/Runtime/StateTree/Source/StateTreeEditorModule/Private/Customizations/StateTreeBindingExtension.cpp b/Engine/Plugins/Runtime/StateTree/Source/StateTreeEditorModule/Private/Customizations/StateTreeBindingExtension.cpp
--- a/Engine/Plugins/Runtime/StateTree/Source/StateTreeEditorModule/Private/Customizations/StateTreeBindingExtension.cpp
+++ b/Engine/Plugins/Runtime/StateTree/Source/StateTreeEditorModule/Private/Customizations/StateTreeBindingExtension.cpp
@@ -38,7 +38,57 @@ UObject* FindEditorBindingsOwner(UObject* InObject)
 	return Result;
 }
 
-void MakeStructPropertyPathFromBindingChain(const FGuid StructID, const TArray<FBindingChainElement>& InBindingChain, FStateTreePropertyPath& OutPath)
+UStruct* ResolveLeafValueStructType(FStateTreeDataView ValueView, const TArray<FBindingChainElement>& InBindingChain)
+{
+	FStateTreePropertyPath Path;
+
+	for (const FBindingChainElement& Element : InBindingChain)
+	{
+		if (const FProperty* Property = Element.Field.Get<FProperty>())
+		{
+			Path.AddPathSegment(Property->GetFName(), Element.ArrayIndex);
+		}
+		else if (const UFunction* Function = Element.Field.Get<UFunction>())
+		{
+			// Cannot handle function calls
+			return nullptr;
+		}
+	}
+
+	TArray<FStateTreePropertyPathIndirection> Indirections;
+	if (!Path.ResolveIndirectionsWithValue(ValueView, Indirections)
+		|| Indirections.IsEmpty())
+	{
+		return nullptr;
+	}
+
+	// Last indirection points to the value of the leaf property, check the type.
+	const FStateTreePropertyPathIndirection& LastIndirection = Indirections.Last();
+
+	UStruct* Result = nullptr;
+
+	if (const FStructProperty* StructProperty = CastField<FStructProperty>(LastIndirection.GetProperty()))
+	{
+		// Get the type of the instanced struct's value.
+		if (StructProperty->Struct == TBaseStructure<FInstancedStruct>::Get())
+		{
+			const FInstancedStruct& InstancedStruct = *reinterpret_cast<const FInstancedStruct*>(LastIndirection.GetPropertyAddress());
+			Result = const_cast<UScriptStruct*>(InstancedStruct.GetScriptStruct());
+		}
+	}
+	else if (const FObjectProperty* ObjectProperty = CastField<FObjectProperty>(LastIndirection.GetProperty()))
+	{
+		// Get type of the instanced object.
+		if (const UObject* Object = *reinterpret_cast<UObject* const*>(LastIndirection.GetPropertyAddress()))
+		{
+			Result = Object->GetClass();
+		}
+	}
+
+	return Result;
+}
+
+void MakeStructPropertyPathFromBindingChain(const FGuid StructID, const TArray<FBindingChainElement>& InBindingChain, FStateTreeDataView DataView, FStateTreePropertyPath& OutPath)
 {
 	OutPath.Reset();
 	OutPath.SetStructID(StructID);
@@ -54,6 +104,8 @@ void MakeStructPropertyPathFromBindingChain(const FGuid StructID, const TArray<F
 			OutPath.AddPathSegment(Function->GetFName());
 		}
 	}
+
+	OutPath.UpdateInstanceStructsFromValue(DataView);
 }
 
 EStateTreePropertyUsage MakeStructPropertyPathFromPropertyHandle(TSharedPtr<const IPropertyHandle> InPropertyHandle, FStateTreePropertyPath& OutPath)
@@ -245,6 +297,8 @@ void FStateTreeBindingExtension::ExtendWidgetRow(FDetailWidgetRow& InWidgetRow,
 	// The struct and property where we're binding.
 	FStateTreePropertyPath TargetPath;
 
+	IStateTreeEditorPropertyBindingsOwner* BindingOwner = nullptr;
+
 	TArray<UObject*> OuterObjects;
 	InPropertyHandle->GetOuterObjects(OuterObjects);
 	if (OuterObjects.Num() == 1)
@@ -255,7 +309,8 @@ void FStateTreeBindingExtension::ExtendWidgetRow(FDetailWidgetRow& InWidgetRow,
 		// Figure out the structs we're editing, and property path relative to current property.
 		UE::StateTree::PropertyBinding::MakeStructPropertyPathFromPropertyHandle(InPropertyHandle, TargetPath);
 
-		if (IStateTreeEditorPropertyBindingsOwner* BindingOwner = Cast<IStateTreeEditorPropertyBindingsOwner>(OwnerObject))
+		BindingOwner = Cast<IStateTreeEditorPropertyBindingsOwner>(OwnerObject);
+		if (BindingOwner)
 		{
 			EditorBindings = BindingOwner->GetPropertyEditorBindings();
 			BindingOwner->GetAccessibleStructs(TargetPath.GetStructID(), AccessibleStructs);
@@ -376,13 +431,19 @@ void FStateTreeBindingExtension::ExtendWidgetRow(FDetailWidgetRow& InWidgetRow,
 			return bCanBind;
 		});
 
-	Args.OnCanBindToContextStruct = FOnCanBindToContextStruct::CreateLambda([InPropertyHandle](const UStruct* InStruct)
+	Args.OnCanBindToContextStruct = FOnCanBindToContextStruct::CreateLambda([InPropertyHandle, AccessibleStructs](const UStruct* InStruct)
 		{
 			// Do not allow to bind directly StateTree nodes
+			// @todo: find a way to more specifically call out the context structs, e.g. pass the property path to the callback.
 			if (InStruct != nullptr)
 			{
-				if (InStruct->IsChildOf(UStateTreeNodeBlueprintBase::StaticClass())
-					|| InStruct->IsChildOf(FStateTreeNodeBase::StaticStruct()))
+				bool bIsStateTreeNode = AccessibleStructs.ContainsByPredicate([InStruct](const FStateTreeBindableStructDesc& AccessibleStruct)
+				{
+					return (AccessibleStruct.DataSource != EStateTreeBindableStructSource::Context && AccessibleStruct.DataSource != EStateTreeBindableStructSource::Parameter)
+						&& AccessibleStruct.Struct == InStruct;
+				});
+
+				if (bIsStateTreeNode)
 				{
 					return false;
 				}
@@ -410,7 +471,7 @@ void FStateTreeBindingExtension::ExtendWidgetRow(FDetailWidgetRow& InWidgetRow,
 			return true;
 		});
 
-	Args.OnAddBinding = FOnAddBinding::CreateLambda([EditorBindings, OwnerObject, TargetPath, AccessibleStructs](FName InPropertyName, const TArray<FBindingChainElement>& InBindingChain)
+	Args.OnAddBinding = FOnAddBinding::CreateLambda([EditorBindings, OwnerObject, BindingOwner, TargetPath, AccessibleStructs](FName InPropertyName, const TArray<FBindingChainElement>& InBindingChain)
 		{
 			IPropertyAccessEditor& PropertyAccessEditor = IModularFeatures::Get().GetModularFeature<IPropertyAccessEditor>("PropertyAccessEditor");
 			if (EditorBindings && OwnerObject)
@@ -425,9 +486,15 @@ void FStateTreeBindingExtension::ExtendWidgetRow(FDetailWidgetRow& InWidgetRow,
 
 					check(SourceStructIndex >= 0 && SourceStructIndex < AccessibleStructs.Num());
 
+					FStateTreeDataView DataView;
+					if (BindingOwner)
+					{
+						BindingOwner->GetDataViewByID(AccessibleStructs[SourceStructIndex].ID, DataView);
+					}
+
 					// If SourceBindingChain is empty at this stage, it means that the binding points to the source struct itself.
 					FStateTreePropertyPath SourcePath;
-					UE::StateTree::PropertyBinding::MakeStructPropertyPathFromBindingChain(AccessibleStructs[SourceStructIndex].ID, SourceBindingChain, SourcePath);
+					UE::StateTree::PropertyBinding::MakeStructPropertyPathFromBindingChain(AccessibleStructs[SourceStructIndex].ID, SourceBindingChain, DataView, SourcePath);
 						
 					OwnerObject->Modify();
 					EditorBindings->AddPropertyBinding(SourcePath, TargetPath);
@@ -564,6 +631,26 @@ void FStateTreeBindingExtension::ExtendWidgetRow(FDetailWidgetRow& InWidgetRow,
 			return BindingColor;
 		});
 
+	if (BindingOwner)
+	{
+		Args.OnResolveIndirection = FOnResolveIndirection::CreateLambda([BindingOwner, AccessibleStructs](TArray<FBindingChainElement> InBindingChain)
+		{
+			const int32 SourceStructIndex = InBindingChain[0].ArrayIndex;
+			TArray<FBindingChainElement> SourceBindingChain = InBindingChain;
+			SourceBindingChain.RemoveAt(0);
+
+			check(SourceStructIndex >= 0 && SourceStructIndex < AccessibleStructs.Num());
+
+			FStateTreeDataView DataView;
+			if (BindingOwner->GetDataViewByID(AccessibleStructs[SourceStructIndex].ID, DataView))
+			{
+				return UE::StateTree::PropertyBinding::ResolveLeafValueStructType(DataView, InBindingChain);
+			}
+
+			return static_cast<UStruct*>(nullptr);
+		});
+	}
+
 	Args.BindButtonStyle = &FAppStyle::Get().GetWidgetStyle<FButtonStyle>("HoverHintOnly");
 	Args.bAllowNewBindings = false;
 	Args.bAllowArrayElementBindings = false;
diff --git a/Engine/Plugins/Runtime/StateTree/Source/StateTreeEditorModule/Private/StateTreeCompiler.cpp b/Engine/Plugins/Runtime/StateTree/Source/StateTreeEditorModule/Private/StateTreeCompiler.cpp
--- a/Engine/Plugins/Runtime/StateTree/Source/StateTreeEditorModule/Private/StateTreeCompiler.cpp
+++ b/Engine/Plugins/Runtime/StateTree/Source/StateTreeEditorModule/Private/StateTreeCompiler.cpp
@@ -149,6 +149,8 @@ bool FStateTreeCompiler::Compile(UStateTree& InStateTree)
 		return false;
 	}
 
+	EditorData->GetAllStructValues(IDToStructValue);
+
 	// Copy schema the EditorData
 	StateTree->Schema = DuplicateObject(EditorData->Schema, StateTree);
 
@@ -1386,15 +1388,6 @@ bool FStateTreeCompiler::GetAndValidateBindings(const FStateTreeBindableStructDe
 		}
 		const FStateTreeBindableStructDesc& SourceStruct = BindingsCompiler.GetSourceStructDesc(SourceStructIdx);
 
-		// Source path should not have any instanced indirections.
-		if (Binding.GetSourcePath().HasAnyInstancedIndirection())
-		{
-			Log.Reportf(EMessageSeverity::Error, TargetStruct,
-						TEXT("Malformed source property path for binding source property '%s' for target '%s:%s'."),
-						*Binding.GetSourcePath().ToString(), *TargetStruct.Name.ToString(), *Binding.GetTargetPath().ToString());
-			return false;
-		}
-
 		// Update path instance types from latest data. E.g. binding may have been created for instanced object of type FooB, and changed to FooA.
  		// @todo: not liking how this mutates the Binding.TargetPath, but currently we dont track well the instanced object changes.
 
@@ -1425,6 +1418,22 @@ bool FStateTreeCompiler::GetAndValidateBindings(const FStateTreeBindableStructDe
 			return false;
 		}
 
+		if (!IDToStructValue.Contains(SourceStructID))
+		{
+			Log.Reportf(EMessageSeverity::Error, TargetStruct,
+				TEXT("Failed to find value for binding source property '%s' for target '%s:%s'."),
+				*Binding.GetSourcePath().ToString(), *TargetStruct.Name.ToString(), *Binding.GetTargetPath().ToString());
+			return false;
+		}
+
+		if (!Binding.GetMutableSourcePath().UpdateInstanceStructsFromValue(IDToStructValue[SourceStructID]))
+		{
+			Log.Reportf(EMessageSeverity::Error, TargetStruct,
+				TEXT("Malformed target property path for binding source property '%s' for source '%s:%s'."),
+				*Binding.GetSourcePath().ToString(), *TargetStruct.Name.ToString(), *Binding.GetTargetPath().ToString());
+			return false;
+		}
+
 		// Special case fo AnyEnum. StateTreeBindingExtension allows AnyEnums to bind to other enum types.
 		// The actual copy will be done via potential type promotion copy, into the value property inside the AnyEnum.
 		// We amend the paths here to point to the 'Value' property.
diff --git a/Engine/Plugins/Runtime/StateTree/Source/StateTreeEditorModule/Private/StateTreeEditorData.cpp b/Engine/Plugins/Runtime/StateTree/Source/StateTreeEditorModule/Private/StateTreeEditorData.cpp
--- a/Engine/Plugins/Runtime/StateTree/Source/StateTreeEditorModule/Private/StateTreeEditorData.cpp
+++ b/Engine/Plugins/Runtime/StateTree/Source/StateTreeEditorModule/Private/StateTreeEditorData.cpp
@@ -246,8 +246,6 @@ FStateTreeBindableStructDesc UStateTreeEditorData::FindContextData(const UStruct
 
 bool UStateTreeEditorData::GetStructByID(const FGuid StructID, FStateTreeBindableStructDesc& OutStructDesc) const
 {
-	bool bResult = false;
-
 	VisitAllNodes([&OutStructDesc, StructID](const UStateTreeState* State, const FStateTreeBindableStructDesc& Desc, const FStateTreeDataView Value)
 	{
 		if (Desc.ID == StructID)
@@ -261,6 +259,21 @@ bool UStateTreeEditorData::GetStructByID(const FGuid StructID, FStateTreeBindabl
 	return OutStructDesc.IsValid();
 }
 
+bool UStateTreeEditorData::GetDataViewByID(const FGuid StructID, FStateTreeDataView& OutDataView) const
+{
+	VisitAllNodes([&OutDataView, StructID](const UStateTreeState* State, const FStateTreeBindableStructDesc& Desc, const FStateTreeDataView Value)
+	{
+		if (Desc.ID == StructID)
+		{
+			OutDataView = Value;
+			return EStateTreeVisitor::Break;
+		}
+		return EStateTreeVisitor::Continue;
+	});
+
+	return OutDataView.IsValid();
+}
+
 const UStateTreeState* UStateTreeEditorData::GetStateByStructID(const FGuid TargetStructID) const
 {
 	const UStateTreeState* Result = nullptr;
diff --git a/Engine/Plugins/Runtime/StateTree/Source/StateTreeEditorModule/Private/StateTreeEditorPropertyBindings.cpp b/Engine/Plugins/Runtime/StateTree/Source/StateTreeEditorModule/Private/StateTreeEditorPropertyBindings.cpp
--- a/Engine/Plugins/Runtime/StateTree/Source/StateTreeEditorModule/Private/StateTreeEditorPropertyBindings.cpp
+++ b/Engine/Plugins/Runtime/StateTree/Source/StateTreeEditorModule/Private/StateTreeEditorPropertyBindings.cpp
@@ -88,16 +88,35 @@ void FStateTreeEditorPropertyBindings::RemoveUnusedBindings(const TMap<FGuid, co
 
 			// Remove binding if path containing instanced indirections (e.g. instance struct or object) cannot be resolved.
 			// TODO: Try to use core redirect to find new name.
-			const FStateTreeDataView Value = ValidStructs.FindChecked(Binding.GetTargetPath().GetStructID());
-			FString Error;
-			TArray<FStateTreePropertyPathIndirection> Indirections;
-			if (!Binding.GetTargetPath().ResolveIndirectionsWithValue(Value, Indirections, &Error))
 			{
-				UE_LOG(LogStateTree, Verbose, TEXT("Removing binding to %s because binding target path cannot be resolved: %s"),
-					*Binding.GetSourcePath().ToString(), *Error); // Error contains the target path.
-				
-				// Remove
-				return true;
+				const FStateTreeDataView* SourceValue = ValidStructs.Find(Binding.GetSourcePath().GetStructID());
+				if (SourceValue && SourceValue->IsValid())
+				{
+					FString Error;
+					TArray<FStateTreePropertyPathIndirection> Indirections;
+					if (!Binding.GetSourcePath().ResolveIndirectionsWithValue(*SourceValue, Indirections, &Error))
+					{
+						UE_LOG(LogStateTree, Verbose, TEXT("Removing binding to %s because binding source path cannot be resolved: %s"),
+							*Binding.GetSourcePath().ToString(), *Error); // Error contains the target path.
+
+						// Remove
+						return true;
+					}
+				}
+			}
+			
+			{
+				const FStateTreeDataView TargetValue = ValidStructs.FindChecked(Binding.GetTargetPath().GetStructID());
+				FString Error;
+				TArray<FStateTreePropertyPathIndirection> Indirections;
+				if (!Binding.GetTargetPath().ResolveIndirectionsWithValue(TargetValue, Indirections, &Error))
+				{
+					UE_LOG(LogStateTree, Verbose, TEXT("Removing binding to %s because binding target path cannot be resolved: %s"),
+						*Binding.GetSourcePath().ToString(), *Error); // Error contains the target path.
+
+					// Remove
+					return true;
+				}
 			}
 
 			return false;
diff --git a/Engine/Plugins/Runtime/StateTree/Source/StateTreeEditorModule/Public/StateTreeCompiler.h b/Engine/Plugins/Runtime/StateTree/Source/StateTreeEditorModule/Public/StateTreeCompiler.h
--- a/Engine/Plugins/Runtime/StateTree/Source/StateTreeEditorModule/Public/StateTreeCompiler.h
+++ b/Engine/Plugins/Runtime/StateTree/Source/StateTreeEditorModule/Public/StateTreeCompiler.h
@@ -66,6 +66,7 @@ private:
 	UStateTree* StateTree = nullptr;
 	UStateTreeEditorData* EditorData = nullptr;
 	TMap<FGuid, int32> IDToState;
+	TMap<FGuid, const FStateTreeDataView > IDToStructValue;
 	TArray<UStateTreeState*> SourceStates;
 
 	TArray<FInstancedStruct> Nodes;
diff --git a/Engine/Plugins/Runtime/StateTree/Source/StateTreeEditorModule/Public/StateTreeEditorData.h b/Engine/Plugins/Runtime/StateTree/Source/StateTreeEditorModule/Public/StateTreeEditorData.h
--- a/Engine/Plugins/Runtime/StateTree/Source/StateTreeEditorModule/Public/StateTreeEditorData.h
+++ b/Engine/Plugins/Runtime/StateTree/Source/StateTreeEditorModule/Public/StateTreeEditorData.h
@@ -32,6 +32,7 @@ public:
 	// IStateTreeEditorPropertyBindingsOwner
 	virtual void GetAccessibleStructs(const FGuid TargetStructID, TArray<FStateTreeBindableStructDesc>& OutStructDescs) const override;
 	virtual bool GetStructByID(const FGuid StructID, FStateTreeBindableStructDesc& OutStructDesc) const override;
+	virtual bool GetDataViewByID(const FGuid StructID, FStateTreeDataView& OutDataView) const override;
 	virtual FStateTreeEditorPropertyBindings* GetPropertyEditorBindings() override { return &EditorBindings; }
 	// ~IStateTreeEditorPropertyBindingsOwner
 
diff --git a/Engine/Plugins/Runtime/StateTree/Source/StateTreeEditorModule/Public/StateTreeEditorPropertyBindings.h b/Engine/Plugins/Runtime/StateTree/Source/StateTreeEditorModule/Public/StateTreeEditorPropertyBindings.h
--- a/Engine/Plugins/Runtime/StateTree/Source/StateTreeEditorModule/Public/StateTreeEditorPropertyBindings.h
+++ b/Engine/Plugins/Runtime/StateTree/Source/StateTreeEditorModule/Public/StateTreeEditorPropertyBindings.h
@@ -113,6 +113,14 @@ class STATETREEEDITORMODULE_API IStateTreeEditorPropertyBindingsOwner
 	 */
 	virtual bool GetStructByID(const FGuid StructID, FStateTreeBindableStructDesc& OutStructDesc) const PURE_VIRTUAL(IStateTreeEditorPropertyBindingsOwner::GetStructByID, return false; );
 
+	/**
+	 * Returns data view based on struct ID.
+	 * @param StructID Target struct ID
+	 * @param OutDataView Result data view.
+	 * @return True if struct found.
+	 */
+	virtual bool GetDataViewByID(const FGuid StructID, FStateTreeDataView& OutDataView) const PURE_VIRTUAL(IStateTreeEditorPropertyBindingsOwner::GetDataViewByID, return false; );
+
 	/** @return Pointer to editor property bindings. */
 	virtual FStateTreeEditorPropertyBindings* GetPropertyEditorBindings() PURE_VIRTUAL(IStateTreeEditorPropertyBindingsOwner::GetPropertyEditorBindings, return nullptr; );
 };
diff --git a/Engine/Source/Editor/UnrealEd/Public/IPropertyAccessEditor.h b/Engine/Source/Editor/UnrealEd/Public/IPropertyAccessEditor.h
--- a/Engine/Source/Editor/UnrealEd/Public/IPropertyAccessEditor.h
+++ b/Engine/Source/Editor/UnrealEd/Public/IPropertyAccessEditor.h
@@ -135,6 +135,9 @@ DECLARE_DELEGATE_RetVal_OneParam(bool, FOnCanRemoveBinding, FName /*InPropertyNa
 /** Delegate called once a new function binding has been created */
 DECLARE_DELEGATE_TwoParams(FOnNewFunctionBindingCreated, UEdGraph* /*InFunctionGraph*/, UFunction* /*InFunction*/);
 
+/** Delegate called to resolve true type of the instance */
+DECLARE_DELEGATE_RetVal_OneParam(UStruct*, FOnResolveIndirection, const TArray<FBindingChainElement>& /*InBindingChain*/);
+
 /** Delegate called once a drag-drop event is dropped on the binding widget */
 DECLARE_DELEGATE_RetVal_TwoParams(FReply, FOnDrop, const FGeometry&, const FDragDropEvent&);
 
@@ -185,6 +188,9 @@ struct FPropertyBindingWidgetArgs
 
 	/** Delegate called once a new function binding has been created */
 	FOnNewFunctionBindingCreated OnNewFunctionBindingCreated;
+
+	/** Delegate called to resolve true type of the instance */
+	FOnResolveIndirection OnResolveIndirection;
 	
 	/** Delegate called when a property is dropped on the property binding widget */
 	FOnDrop OnDrop;
