Fix being unable to bind to interface functions on anim nodes

diff --git a/Engine/Plugins/Runtime/PropertyAccess/Source/PropertyAccessEditor/Private/SPropertyBinding.cpp b/Engine/Plugins/Runtime/PropertyAccess/Source/PropertyAccessEditor/Private/SPropertyBinding.cpp
--- a/Engine/Plugins/Runtime/PropertyAccess/Source/PropertyAccessEditor/Private/SPropertyBinding.cpp
+++ b/Engine/Plugins/Runtime/PropertyAccess/Source/PropertyAccessEditor/Private/SPropertyBinding.cpp
@@ -120,7 +120,7 @@ void SPropertyBinding::ForEachBindableFunction(UClass* FromClass, Predicate Pred
 			}
 		
 			// Walk up class hierarchy for native functions and properties
-			for ( TFieldIterator<UFunction> FuncIt(InBindableClass, EFieldIteratorFlags::IncludeSuper); FuncIt; ++FuncIt )
+			for ( TFieldIterator<UFunction> FuncIt(InBindableClass, EFieldIteratorFlags::IncludeSuper, EFieldIteratorFlags::ExcludeDeprecated, EFieldIteratorFlags::IncludeInterfaces); FuncIt; ++FuncIt )
 			{
 				UFunction* Function = *FuncIt;
 
