Fix bug with property binding widget where it won't display anything if the only bindable properties are children of a UObject function

diff --git a/Engine/Plugins/Runtime/PropertyAccess/Source/PropertyAccessEditor/Private/SPropertyBinding.cpp b/Engine/Plugins/Runtime/PropertyAccess/Source/PropertyAccessEditor/Private/SPropertyBinding.cpp
--- a/Engine/Plugins/Runtime/PropertyAccess/Source/PropertyAccessEditor/Private/SPropertyBinding.cpp
+++ b/Engine/Plugins/Runtime/PropertyAccess/Source/PropertyAccessEditor/Private/SPropertyBinding.cpp
@@ -360,11 +360,20 @@ bool SPropertyBinding::HasBindablePropertiesRecursive(UStruct* InStruct, TSet<US
 			{
 				if(FObjectPropertyBase* ObjectPropertyBase = CastField<FObjectPropertyBase>(ReturnProperty))
 				{
-					HasBindablePropertiesRecursive(ObjectPropertyBase->PropertyClass, VisitedStructs, BindingChain);
+					if (HasBindablePropertiesRecursive(ObjectPropertyBase->PropertyClass, VisitedStructs, BindingChain))
+					{
+						if (Args.bAllowUObjectFunctions)
+						{
+							BindableCount++;
+						}
+					}
 				}
 				else if(FStructProperty* StructProperty = CastField<FStructProperty>(ReturnProperty))
 				{
-					HasBindablePropertiesRecursive(StructProperty->Struct, VisitedStructs, BindingChain);
+					if (HasBindablePropertiesRecursive(StructProperty->Struct, VisitedStructs, BindingChain))
+					{
+						BindableCount++;
+					}
 				}
 			}
 		});
