Fix crash opening binding menus in state tree

Removal of guard in CL 28311795 means that lower-level code can potentially crash - found by jakub.kuznik [at] CDPR.

diff --git a/Engine/Plugins/Runtime/PropertyAccess/Source/PropertyAccessEditor/Private/SPropertyBinding.cpp b/Engine/Plugins/Runtime/PropertyAccess/Source/PropertyAccessEditor/Private/SPropertyBinding.cpp
--- a/Engine/Plugins/Runtime/PropertyAccess/Source/PropertyAccessEditor/Private/SPropertyBinding.cpp
+++ b/Engine/Plugins/Runtime/PropertyAccess/Source/PropertyAccessEditor/Private/SPropertyBinding.cpp
@@ -201,7 +201,7 @@ void SPropertyBinding::ForEachBindableFunction(UClass* FromClass, Predicate Pred
 template <typename Predicate>
 void SPropertyBinding::ForEachBindableProperty(UStruct* InStruct, Predicate Pred) const
 {
-	if(Args.OnCanBindProperty.IsBound())
+	if(InStruct && Args.OnCanBindProperty.IsBound())
 	{
 		UBlueprintGeneratedClass* SkeletonClass = Blueprint ? Cast<UBlueprintGeneratedClass>(Blueprint->SkeletonGeneratedClass) : nullptr;
 
@@ -887,7 +887,7 @@ void SPropertyBinding::FillPropertyMenu(FMenuBuilder& MenuBuilder, UStruct* InOw
 	}
 
 	// Add 'none' entry only if we just have the search block in the builder
-	if ( MenuBuilder.GetMultiBox()->GetBlocks().Num() == 1 )
+	if ( InOwnerStruct && MenuBuilder.GetMultiBox()->GetBlocks().Num() == 1 )
 	{
 		MenuBuilder.BeginSection("None", InOwnerStruct->GetDisplayNameText());
 		MenuBuilder.AddWidget(SNew(STextBlock).Text(LOCTEXT("None", "None")), FText::GetEmpty());
