PropertyAccess: Added a version of the OnCanAcceptProperty which gets the property chain passed in.

diff --git a/Engine/Plugins/Experimental/Chooser/Source/ChooserEditor/Private/SPropertyAccessChainWidget.cpp b/Engine/Plugins/Experimental/Chooser/Source/ChooserEditor/Private/SPropertyAccessChainWidget.cpp
--- a/Engine/Plugins/Experimental/Chooser/Source/ChooserEditor/Private/SPropertyAccessChainWidget.cpp
+++ b/Engine/Plugins/Experimental/Chooser/Source/ChooserEditor/Private/SPropertyAccessChainWidget.cpp
@@ -143,7 +143,7 @@ TSharedRef<SWidget> SPropertyAccessChainWidget::CreatePropertyAccessWidget()
 			return true;
 		});
 
-	Args.OnCanAcceptPropertyOrChildren = FOnCanBindProperty::CreateLambda([](FProperty* InProperty)
+	Args.OnCanAcceptPropertyOrChildrenWithBindingChain = FOnCanAcceptPropertyOrChildrenWithBindingChain::CreateLambda([](FProperty* InProperty, TConstArrayView<TSharedPtr<FBindingChainElement>> BindingChain)
 		{
 			// Make only blueprint visible properties visible for binding.
 			return InProperty->HasAnyPropertyFlags(CPF_BlueprintVisible);
diff --git a/Engine/Plugins/Runtime/PropertyAccess/Source/PropertyAccessEditor/Private/SPropertyBinding.cpp b/Engine/Plugins/Runtime/PropertyAccess/Source/PropertyAccessEditor/Private/SPropertyBinding.cpp
--- a/Engine/Plugins/Runtime/PropertyAccess/Source/PropertyAccessEditor/Private/SPropertyBinding.cpp
+++ b/Engine/Plugins/Runtime/PropertyAccess/Source/PropertyAccessEditor/Private/SPropertyBinding.cpp
@@ -199,7 +199,7 @@ void SPropertyBinding::ForEachBindableFunction(UClass* FromClass, Predicate Pred
 }
 
 template <typename Predicate>
-void SPropertyBinding::ForEachBindableProperty(UStruct* InStruct, Predicate Pred) const
+void SPropertyBinding::ForEachBindableProperty(UStruct* InStruct, TConstArrayView<TSharedPtr<FBindingChainElement>> BindingChain, Predicate Pred) const
 {
 	if(InStruct && Args.OnCanBindProperty.IsBound())
 	{
@@ -215,10 +215,18 @@ void SPropertyBinding::ForEachBindableProperty(UStruct* InStruct, Predicate Pred
 				break;
 			}
 
+			PRAGMA_DISABLE_DEPRECATION_WARNINGS
 			if (Args.OnCanAcceptPropertyOrChildren.IsBound() && Args.OnCanAcceptPropertyOrChildren.Execute(Property) == false)
 			{
 				continue;
 			}
+			PRAGMA_ENABLE_DEPRECATION_WARNINGS
+			
+			if (Args.OnCanAcceptPropertyOrChildrenWithBindingChain.IsBound()
+				&& Args.OnCanAcceptPropertyOrChildrenWithBindingChain.Execute(Property, BindingChain) == false)
+			{
+				continue;
+			}
 
 			if (SkeletonClass)
 			{
@@ -268,7 +276,7 @@ bool SPropertyBinding::HasBindablePropertiesRecursive(UStruct* InStruct, TSet<US
 	}
 	
 	int32 BindableCount = 0;
-	ForEachBindableProperty(InStruct, [this, &BindableCount, &VisitedStructs, &BindingChain] (FProperty* Property)
+	ForEachBindableProperty(InStruct, BindingChain, [this, &BindableCount, &VisitedStructs, &BindingChain] (FProperty* Property)
 	{
 		BindingChain.Emplace(MakeShared<FBindingChainElement>(Property));
 		ON_SCOPE_EXIT{ BindingChain.Pop(); };
@@ -799,7 +807,7 @@ void SPropertyBinding::FillPropertyMenu(FMenuBuilder& MenuBuilder, UStruct* InOw
 
 			MenuBuilder.BeginSection("Properties", LOCTEXT("Properties", "Properties"));
 			{
-				ForEachBindableProperty(InOwnerStruct, [this, &InBindingChain, BindingStruct, &MenuBuilder, &MakeArrayElementPropertyWidget, &MakePropertyWidget, &MakePropertyEntry, &MakeArrayElementEntry] (FProperty* Property)
+				ForEachBindableProperty(InOwnerStruct, InBindingChain, [this, &InBindingChain, BindingStruct, &MenuBuilder, &MakeArrayElementPropertyWidget, &MakePropertyWidget, &MakePropertyEntry, &MakeArrayElementEntry] (FProperty* Property)
 				{
 					TArray<TSharedPtr<FBindingChainElement>> NewBindingChain(InBindingChain);
 					NewBindingChain.Emplace(MakeShared<FBindingChainElement>(Property));
diff --git a/Engine/Plugins/Runtime/PropertyAccess/Source/PropertyAccessEditor/Private/SPropertyBinding.h b/Engine/Plugins/Runtime/PropertyAccess/Source/PropertyAccessEditor/Private/SPropertyBinding.h
--- a/Engine/Plugins/Runtime/PropertyAccess/Source/PropertyAccessEditor/Private/SPropertyBinding.h
+++ b/Engine/Plugins/Runtime/PropertyAccess/Source/PropertyAccessEditor/Private/SPropertyBinding.h
@@ -77,7 +77,7 @@ private:
 	bool HasBindablePropertiesRecursive(UStruct* InStruct, TSet<UStruct*>& VisitedStructs, TArray<TSharedPtr<FBindingChainElement>>& BindingChain) const;
 	
 	template <typename Predicate>
-	void ForEachBindableProperty(UStruct* InStruct, Predicate Pred) const;
+	void ForEachBindableProperty(UStruct* InStruct, TConstArrayView<TSharedPtr<FBindingChainElement>> BindingChain, Predicate Pred) const;
 
 	template <typename Predicate>
 	void ForEachBindableFunction(UClass* FromClass, Predicate Pred) const;
diff --git a/Engine/Plugins/Runtime/StateTree/Source/StateTreeEditorModule/Private/Customizations/StateTreeBindingExtension.cpp b/Engine/Plugins/Runtime/StateTree/Source/StateTreeEditorModule/Private/Customizations/StateTreeBindingExtension.cpp
--- a/Engine/Plugins/Runtime/StateTree/Source/StateTreeEditorModule/Private/Customizations/StateTreeBindingExtension.cpp
+++ b/Engine/Plugins/Runtime/StateTree/Source/StateTreeEditorModule/Private/Customizations/StateTreeBindingExtension.cpp
@@ -845,7 +845,7 @@ void FStateTreeBindingExtension::ExtendWidgetRow(FDetailWidgetRow& InWidgetRow,
 			return CachedBindingData->CanBindToContextStruct(InStruct);
 		});
 
-	Args.OnCanAcceptPropertyOrChildren = FOnCanBindProperty::CreateLambda([](FProperty* InProperty)
+	Args.OnCanAcceptPropertyOrChildrenWithBindingChain = FOnCanAcceptPropertyOrChildrenWithBindingChain::CreateLambda([](FProperty* InProperty, TConstArrayView<TSharedPtr<FBindingChainElement>> InBindingChain)
 		{
 			// Make only editor visible properties visible for binding.
 			return InProperty->HasAnyPropertyFlags(CPF_Edit);
diff --git a/Engine/Source/Editor/UnrealEd/Public/IPropertyAccessEditor.h b/Engine/Source/Editor/UnrealEd/Public/IPropertyAccessEditor.h
--- a/Engine/Source/Editor/UnrealEd/Public/IPropertyAccessEditor.h
+++ b/Engine/Source/Editor/UnrealEd/Public/IPropertyAccessEditor.h
@@ -106,6 +106,9 @@ DECLARE_DELEGATE_RetVal_OneParam(bool, FOnGotoBinding, FName /*InPropertyName*/)
 DECLARE_DELEGATE_RetVal_OneParam(bool, FOnCanGotoBinding, FName /*InPropertyName*/);
 
 /** Delegate used to check whether a property is considered for binding. Returning false will discard the property and all child properties. */
+DECLARE_DELEGATE_RetVal_TwoParams(bool, FOnCanAcceptPropertyOrChildrenWithBindingChain, FProperty* /*InProperty*/, TConstArrayView<TSharedPtr<FBindingChainElement>> /*InBindingChain*/);
+
+// UE_DEPRECATED(5.4, "Please use OnCanAcceptPropertyOrChildrenWithBindingChain instead.")
 DECLARE_DELEGATE_RetVal_OneParam(bool, FOnCanAcceptPropertyOrChildren, FProperty* /*InProperty*/);
 
 /** Delegate used to check whether a property can be bound to the property in question */
@@ -144,6 +147,15 @@ DECLARE_DELEGATE_RetVal_TwoParams(FReply, FOnDrop, const FGeometry&, const FDrag
 /** Setup arguments structure for a property binding widget */
 struct FPropertyBindingWidgetArgs
 {
+	// Macro needed to avoid deprecation errors when the struct is copied or created in the default methods.
+	PRAGMA_DISABLE_DEPRECATION_WARNINGS
+	FPropertyBindingWidgetArgs() = default;
+	FPropertyBindingWidgetArgs(const FPropertyBindingWidgetArgs&) = default;
+	FPropertyBindingWidgetArgs(FPropertyBindingWidgetArgs&&) = default;
+	FPropertyBindingWidgetArgs& operator=(const FPropertyBindingWidgetArgs&) = default;
+	FPropertyBindingWidgetArgs& operator=(FPropertyBindingWidgetArgs&&) = default;
+	PRAGMA_ENABLE_DEPRECATION_WARNINGS
+	
 	/** An optional bindable property */
 	FProperty* Property = nullptr;
 
@@ -160,6 +172,9 @@ struct FPropertyBindingWidgetArgs
 	FOnCanGotoBinding OnCanGotoBinding;
 
 	/** Delegate used to check whether a property is considered for binding. Returning false will discard the property and all child properties. */
+	FOnCanAcceptPropertyOrChildrenWithBindingChain OnCanAcceptPropertyOrChildrenWithBindingChain;
+
+	UE_DEPRECATED(5.4, "Please use OnCanAcceptPropertyOrChildrenWithBindingChain instead.")
 	FOnCanAcceptPropertyOrChildren OnCanAcceptPropertyOrChildren;
 	
 	/** Delegate used to check whether a property can be bound to the property in question */
