Property Bindings: Add MaxDepth parameter to the widget so that systems can customize the maximum property depth to discover

diff --git a/Engine/Plugins/Runtime/PropertyAccess/Source/PropertyAccessEditor/Private/SPropertyBinding.cpp b/Engine/Plugins/Runtime/PropertyAccess/Source/PropertyAccessEditor/Private/SPropertyBinding.cpp
--- a/Engine/Plugins/Runtime/PropertyAccess/Source/PropertyAccessEditor/Private/SPropertyBinding.cpp
+++ b/Engine/Plugins/Runtime/PropertyAccess/Source/PropertyAccessEditor/Private/SPropertyBinding.cpp
@@ -270,7 +270,7 @@ bool SPropertyBinding::HasBindablePropertiesRecursive(UStruct* InStruct, TSet<US
 	VisitedStructs.Add(InStruct);
 	
 	// Arbitrary cut off to avoid infinite loops.
-	if (BindingChain.Num() > 10)
+	if (BindingChain.Num() > Args.MaxDepth)
 	{
 		return false;
 	}
diff --git a/Engine/Source/Editor/UnrealEd/Public/IPropertyAccessEditor.h b/Engine/Source/Editor/UnrealEd/Public/IPropertyAccessEditor.h
--- a/Engine/Source/Editor/UnrealEd/Public/IPropertyAccessEditor.h
+++ b/Engine/Source/Editor/UnrealEd/Public/IPropertyAccessEditor.h
@@ -227,6 +227,9 @@ struct FPropertyBindingWidgetArgs
 
 	/** Optional style override for bind button */
 	const FButtonStyle* BindButtonStyle = nullptr;
+
+	/** The maximum level of depth to generate */
+	uint8 MaxDepth = 10;
 	
 	/** Whether to generate pure bindings */
 	bool bGeneratePureBindings = true;
