[PropertyBindings] fixed potential crash when parsing long binding chains.
The issue was that an ArrayView was used to pass the BindingChain around but the predicate can modify that array which invalidates the ArrayView if reallocation is performed.

diff --git a/Engine/Plugins/Runtime/PropertyAccess/Source/PropertyAccessEditor/Private/SPropertyBinding.cpp b/Engine/Plugins/Runtime/PropertyAccess/Source/PropertyAccessEditor/Private/SPropertyBinding.cpp
--- a/Engine/Plugins/Runtime/PropertyAccess/Source/PropertyAccessEditor/Private/SPropertyBinding.cpp
+++ b/Engine/Plugins/Runtime/PropertyAccess/Source/PropertyAccessEditor/Private/SPropertyBinding.cpp
@@ -199,7 +199,7 @@ void SPropertyBinding::ForEachBindableFunction(UClass* FromClass, Predicate Pred
 }
 
 template <typename Predicate>
-void SPropertyBinding::ForEachBindableProperty(UStruct* InStruct, TConstArrayView<TSharedPtr<FBindingChainElement>> BindingChain, Predicate Pred) const
+void SPropertyBinding::ForEachBindableProperty(UStruct* InStruct, const TArray<TSharedPtr<FBindingChainElement>>& BindingChain, Predicate Pred) const
 {
 	PRAGMA_DISABLE_DEPRECATION_WARNINGS
 	const bool bIsOnCanBindPropertyBound = Args.OnCanBindProperty.IsBound() || Args.OnCanBindPropertyWithBindingChain.IsBound();
diff --git a/Engine/Plugins/Runtime/PropertyAccess/Source/PropertyAccessEditor/Private/SPropertyBinding.h b/Engine/Plugins/Runtime/PropertyAccess/Source/PropertyAccessEditor/Private/SPropertyBinding.h
--- a/Engine/Plugins/Runtime/PropertyAccess/Source/PropertyAccessEditor/Private/SPropertyBinding.h
+++ b/Engine/Plugins/Runtime/PropertyAccess/Source/PropertyAccessEditor/Private/SPropertyBinding.h
@@ -81,14 +81,18 @@ private:
 	bool IsFieldFromDeniedClass(FFieldVariant Field) const;
 	bool HasBindableProperties(UStruct* InStruct, TArray<TSharedPtr<FBindingChainElement>>& BindingChain) const;
 	bool HasBindablePropertiesRecursive(UStruct* InStruct, TSet<UStruct*>& VisitedStructs, TArray<TSharedPtr<FBindingChainElement>>& BindingChain) const;
-	
+
+	/**
+	 * Note that an ArrayView is not used to pass the BindingChain around since the predicate can modify the array
+	 * and this will invalidate the ArrayView if reallocation is performed.
+	 */
 	template <typename Predicate>
-	void ForEachBindableProperty(UStruct* InStruct, TConstArrayView<TSharedPtr<FBindingChainElement>> BindingChain, Predicate Pred) const;
+	void ForEachBindableProperty(UStruct* InStruct, const TArray<TSharedPtr<FBindingChainElement>>& BindingChain, Predicate Pred) const;
 
 	template <typename Predicate>
 	void ForEachBindableFunction(UClass* FromClass, Predicate Pred) const;
 
-	UBlueprint* Blueprint;
+	UBlueprint* Blueprint = nullptr;
 	TArray<FBindingContextStruct> BindingContextStructs;
 	FPropertyBindingWidgetArgs Args;
 	FName PropertyName;
