Fix grammar in comment

diff --git a/Engine/Plugins/Runtime/PropertyAccess/Source/PropertyAccessEditor/Private/SPropertyBinding.cpp b/Engine/Plugins/Runtime/PropertyAccess/Source/PropertyAccessEditor/Private/SPropertyBinding.cpp
--- a/Engine/Plugins/Runtime/PropertyAccess/Source/PropertyAccessEditor/Private/SPropertyBinding.cpp
+++ b/Engine/Plugins/Runtime/PropertyAccess/Source/PropertyAccessEditor/Private/SPropertyBinding.cpp
@@ -397,7 +397,7 @@ TSharedRef<SWidget> SPropertyBinding::OnGenerateDelegateMenu()
 	// The menu itself is be searchable.
 	const bool bSearchableMenu = true;
 
-	// The menu are generated through reflection and sometime the API exposes some recursivity (think about a Widget returning it parent which is also a Widget). Just by reflection
+	// The menu are generated through reflection and sometime the API exposes some recursivity (think about a Widget returning its parent which is also a Widget). Just by reflection
 	// it is not possible to determine when the root object is reached. It needs a kind of simulation which is not implemented. Also, even if the recursivity was correctly handled, the possible
 	// permutations tend to grow exponentially. Until a clever solution is found, the simple approach is to disable recursively searching those menus. User can still search the current one though.
 	const bool bRecursivelySearchableMenu = false;
