PropertyBindings: Introduced overriding icon color. (contributed)

diff --git a/Engine/Plugins/Runtime/PropertyAccess/Source/PropertyAccessEditor/Private/SPropertyBinding.cpp b/Engine/Plugins/Runtime/PropertyAccess/Source/PropertyAccessEditor/Private/SPropertyBinding.cpp
--- a/Engine/Plugins/Runtime/PropertyAccess/Source/PropertyAccessEditor/Private/SPropertyBinding.cpp
+++ b/Engine/Plugins/Runtime/PropertyAccess/Source/PropertyAccessEditor/Private/SPropertyBinding.cpp
@@ -465,9 +465,14 @@ TSharedRef<SWidget> SPropertyBinding::OnGenerateDelegateMenu()
 					PinType.PinSubCategory = NAME_None;
 					PinType.PinSubCategoryObject = ScriptStruct;
 				}
-				Icon = FBlueprintEditorUtils::GetIconFromPin(PinType, true); 
+				Icon = FBlueprintEditorUtils::GetIconFromPin(PinType, true);
 				IconColor = Schema->GetPinTypeColor(PinType);
 			}
+
+			if (ContextStruct.Color.IsSet())
+			{
+				IconColor = ContextStruct.Color.GetValue();
+			}
 			
 			return SNew(SHorizontalBox)
 				.ToolTipText(ToolTipText)
diff --git a/Engine/Source/Editor/UnrealEd/Public/IPropertyAccessEditor.h b/Engine/Source/Editor/UnrealEd/Public/IPropertyAccessEditor.h
--- a/Engine/Source/Editor/UnrealEd/Public/IPropertyAccessEditor.h
+++ b/Engine/Source/Editor/UnrealEd/Public/IPropertyAccessEditor.h
@@ -85,6 +85,9 @@ struct FBindingContextStruct
 	/** Icon to display in the popup menu. */ 
 	const FSlateBrush* Icon = nullptr;
 
+	/** Color of the icon to display in the popup menu. */ 
+	TOptional<FLinearColor> Color;
+
 	/** Text to display for this item in the popup. If left empty, struct's display text will be used. */
 	FText DisplayText;
 
