Add support for category inside Property binding

#rnx

diff --git a/Engine/Plugins/Runtime/PropertyAccess/Source/PropertyAccessEditor/Private/SPropertyBinding.cpp b/Engine/Plugins/Runtime/PropertyAccess/Source/PropertyAccessEditor/Private/SPropertyBinding.cpp
--- a/Engine/Plugins/Runtime/PropertyAccess/Source/PropertyAccessEditor/Private/SPropertyBinding.cpp
+++ b/Engine/Plugins/Runtime/PropertyAccess/Source/PropertyAccessEditor/Private/SPropertyBinding.cpp
@@ -33,6 +33,33 @@ void SPropertyBinding::Construct(const FArguments& InArgs, UBlueprint* InBluepri
 	Args = InArgs._Args;
 	PropertyName = Args.Property != nullptr ? Args.Property->GetFName() : NAME_None;
 
+	for(int32 Index = 0; Index < BindingContextStructs.Num(); ++Index)
+	{
+		FBindingContextStruct& ContextStruct = BindingContextStructs[Index];
+		/* BindingContextStructSections contains an entry for each different FBindingContextStruct.Section.
+		 Section are then further divided into subcategories based on the FBindingContextStruct.Category
+		 Section are displayed as a section while subcategories are displayed as menu and submenu under a section. */
+		FBindingContextStructCategory* Section = BindingContextStructSections.FindByPredicate([&](const FBindingContextStructCategory& Section) { return Section.Name.IdenticalTo(ContextStruct.Section); });
+		if(!Section)
+		{
+			Section = &BindingContextStructSections.AddDefaulted_GetRef();
+			Section->Name = ContextStruct.Section;
+		}
+
+		//No category the BindingContextStruct will be displayed directly inside the section.
+		if (ContextStruct.Category.IsEmpty())
+		{
+			Section->BindingContextStructIndices.Add(Index);
+		}
+		else
+		{
+			// Parse the categories and subcategories and construct the data to build the menus.
+			TArray<FString> Categories;
+			ContextStruct.Category.ParseIntoArray(Categories, TEXT("|"));
+			BuildContextStructCategoryRecursive(Categories, Section->SubCategories, Index);
+		}
+	}
+
 	ChildSlot
 	[
 		SNew(SHorizontalBox)
@@ -198,6 +225,144 @@ void SPropertyBinding::ForEachBindableFunction(UClass* FromClass, Predicate Pred
 	}
 }
 
+void SPropertyBinding::AddCategoryToMenu(FMenuBuilder& MenuBuilder, const FBindingContextStructCategory& Category)
+{
+	if (HasCategorySomethingToDisplayRecursive(Category))
+	{
+		TSharedRef<SWidget> SubMenu =
+			SNew(SHorizontalBox)
+			+ SHorizontalBox::Slot()
+				.AutoWidth()
+				[
+					SNew(SSpacer)
+					.Size(FVector2D(18.0f, 0.0f))
+				]
+			+ SHorizontalBox::Slot()
+				.AutoWidth()
+				[
+					SNew(STextBlock)
+					.Text(Category.Name)
+				];
+
+		MenuBuilder.AddSubMenu(SubMenu, FNewMenuDelegate::CreateSP(this, &SPropertyBinding::FillCategoryMenu, &Category));
+	}
+}
+
+void SPropertyBinding::BuildContextStructCategoryRecursive(TConstArrayView<FString> CategoryNames, TArray<FBindingContextStructCategory>& ParentSubCategories, int32 ContextStructIndex)
+{
+	if (CategoryNames.IsEmpty())
+	{
+		return;
+	}
+
+	FString CategoryName = CategoryNames[0];
+	FBindingContextStructCategory* CategoryPtr = ParentSubCategories.FindByPredicate([&](const FBindingContextStructCategory& Category) { return Category.Name.ToString() == CategoryName; });
+	if (!CategoryPtr)
+	{
+		CategoryPtr = &ParentSubCategories.AddDefaulted_GetRef();
+		CategoryPtr->Name = FText::FromString(CategoryName);
+	}
+	
+	if(CategoryNames.Num() == 1)
+	{
+		CategoryPtr->BindingContextStructIndices.Add(ContextStructIndex);
+	}
+	else
+	{
+		BuildContextStructCategoryRecursive(CategoryNames.RightChop(1), CategoryPtr->SubCategories, ContextStructIndex);
+	}
+}
+
+bool SPropertyBinding::HasCategorySomethingToDisplayRecursive(const FBindingContextStructCategory& Category) const
+{
+	for(const FBindingContextStructCategory& SubCategory : Category.SubCategories)
+	{
+		if (HasCategorySomethingToDisplayRecursive(SubCategory))
+		{
+			return true;
+		}
+	}
+
+	for(const int32 ContextStructIndex : Category.BindingContextStructIndices)
+	{
+		const FBindingContextStruct& ContextStruct = BindingContextStructs[ContextStructIndex];
+		UStruct* Struct = const_cast<UStruct*>(ContextStruct.Struct);
+
+		PRAGMA_DISABLE_DEPRECATION_WARNINGS
+		const bool bDeprecatedCanBindToContextStruct = Args.OnCanBindToContextStruct.IsBound() && Args.OnCanBindToContextStruct.Execute(Struct);
+		PRAGMA_ENABLE_DEPRECATION_WARNINGS
+
+		if (bDeprecatedCanBindToContextStruct || (Args.OnCanBindToContextStructWithIndex.IsBound() && Args.OnCanBindToContextStructWithIndex.Execute(Struct, ContextStructIndex)))
+		{
+			return true;
+		}
+
+		TArray<TSharedPtr<FBindingChainElement>> BindingChain;
+		BindingChain.Emplace(MakeShared<FBindingChainElement>(nullptr, ContextStructIndex));
+		if (HasBindableProperties(Struct, BindingChain))
+		{
+			return true;
+		}
+	}
+
+	return false;
+}
+
+TSharedRef<SWidget> SPropertyBinding::MakeContextStructWidget(const FBindingContextStruct& ContextStruct) const
+{
+	const FText DisplayText = ContextStruct.DisplayText.IsEmpty() ? ContextStruct.Struct->GetDisplayNameText() : ContextStruct.DisplayText;
+	const FText ToolTipText = ContextStruct.TooltipText.IsEmpty() ? DisplayText : ContextStruct.TooltipText;
+
+	const FSlateBrush* Icon = ContextStruct.Icon;
+	FLinearColor IconColor = FLinearColor::White;
+	if (Icon == nullptr)
+	{
+		const UEdGraphSchema_K2* Schema = GetDefault<UEdGraphSchema_K2>();
+
+		FEdGraphPinType PinType;
+
+		if (UClass* Class = Cast<UClass>(ContextStruct.Struct))
+		{
+			PinType.PinCategory = UEdGraphSchema_K2::PC_Object;
+			PinType.PinSubCategory = NAME_None;
+			PinType.PinSubCategoryObject = Class;
+		}
+		else if (UScriptStruct* ScriptStruct = Cast<UScriptStruct>(ContextStruct.Struct))
+		{
+			PinType.PinCategory = UEdGraphSchema_K2::PC_Struct;
+			PinType.PinSubCategory = NAME_None;
+			PinType.PinSubCategoryObject = ScriptStruct;
+		}
+		Icon = FBlueprintEditorUtils::GetIconFromPin(PinType, true);
+		IconColor = Schema->GetPinTypeColor(PinType);
+	}
+
+	if (ContextStruct.Color.IsSet())
+	{
+		IconColor = ContextStruct.Color.GetValue();
+	}
+
+	return SNew(SHorizontalBox)
+			   .ToolTipText(ToolTipText)
+		+ SHorizontalBox::Slot()
+			  .AutoWidth()
+				  [SNew(SSpacer)
+						  .Size(FVector2D(18.0f, 0.0f))]
+		+ SHorizontalBox::Slot()
+			  .AutoWidth()
+			  .VAlign(VAlign_Center)
+			  .Padding(1.0f, 0.0f)
+				  [SNew(SImage)
+						  .Image(Icon)
+						  .ColorAndOpacity(IconColor)]
+		+ SHorizontalBox::Slot()
+			  .AutoWidth()
+			  .VAlign(VAlign_Center)
+			  .Padding(4.0f, 0.0f)
+				  [SNew(STextBlock)
+						  .Text(DisplayText)];
+}
+
 template <typename Predicate>
 void SPropertyBinding::ForEachBindableProperty(UStruct* InStruct, const TArray<TSharedPtr<FBindingChainElement>>& BindingChain, Predicate Pred) const
 {
@@ -440,115 +605,44 @@ TSharedRef<SWidget> SPropertyBinding::OnGenerateDelegateMenu()
 
 	if (BindingContextStructs.Num() > 0)
 	{
-		auto MakeContextStructWidget = [this](const FBindingContextStruct& ContextStruct)
+		// Go thought each top level section and display the sub categories as menu and then the context struct without categories.
+		for(const FBindingContextStructCategory& Section : BindingContextStructSections)
 		{
-			const FText DisplayText = ContextStruct.DisplayText.IsEmpty() ? ContextStruct.Struct->GetDisplayNameText() : ContextStruct.DisplayText;
-			const FText ToolTipText = ContextStruct.TooltipText.IsEmpty() ? DisplayText : ContextStruct.TooltipText;
-
-			const FSlateBrush* Icon = ContextStruct.Icon;
-			FLinearColor IconColor = FLinearColor::White;
-			if (Icon == nullptr)
+			MenuBuilder.BeginSection(FName(), Section.Name);
+			for(const FBindingContextStructCategory& SubCategory : Section.SubCategories)
 			{
-				const UEdGraphSchema_K2* Schema = GetDefault<UEdGraphSchema_K2>();
-
-				FEdGraphPinType PinType;
-
-				if (UClass* Class = Cast<UClass>(ContextStruct.Struct))
-				{
-					PinType.PinCategory = UEdGraphSchema_K2::PC_Object;
-					PinType.PinSubCategory = NAME_None;
-					PinType.PinSubCategoryObject = Class;
-				}
-				else if (UScriptStruct* ScriptStruct = Cast<UScriptStruct>(ContextStruct.Struct))
-				{
-					PinType.PinCategory = UEdGraphSchema_K2::PC_Struct;
-					PinType.PinSubCategory = NAME_None;
-					PinType.PinSubCategoryObject = ScriptStruct;
-				}
-				Icon = FBlueprintEditorUtils::GetIconFromPin(PinType, true);
-				IconColor = Schema->GetPinTypeColor(PinType);
+				AddCategoryToMenu(MenuBuilder, SubCategory);
 			}
 
-			if (ContextStruct.Color.IsSet())
+			for (int32 Index : Section.BindingContextStructIndices)
 			{
-				IconColor = ContextStruct.Color.GetValue();
-			}
-			
-			return SNew(SHorizontalBox)
-				.ToolTipText(ToolTipText)
-				+ SHorizontalBox::Slot()
-				.AutoWidth()
-				[
-					SNew(SSpacer)
-					.Size(FVector2D(18.0f, 0.0f))
-				]
-				+ SHorizontalBox::Slot()
-				.AutoWidth()
-				.VAlign(VAlign_Center)
-				.Padding(1.0f, 0.0f)
-				[
-					SNew(SImage)
-					.Image(Icon)
-					.ColorAndOpacity(IconColor)
-				]
-				+ SHorizontalBox::Slot()
-				.AutoWidth()
-				.VAlign(VAlign_Center)
-				.Padding(4.0f, 0.0f)
-				[
-					SNew(STextBlock)
-					.Text(DisplayText)
-				];
-		};
+				FBindingContextStruct& ContextStruct = BindingContextStructs[Index];
 
-		bool bSectionOpened = false;
-		FText CurrentSection = FText::GetEmpty();
-		
-		for (int32 i = 0; i < BindingContextStructs.Num(); i++)
-		{
-			FBindingContextStruct& ContextStruct = BindingContextStructs[i];
+				// Make first chain element representing the index in the context array.
+				TArray<TSharedPtr<FBindingChainElement>> BindingChain;
+				BindingChain.Emplace(MakeShared<FBindingChainElement>(nullptr, Index));
 
-			if (!ContextStruct.Section.IdenticalTo(CurrentSection))
-			{
-				if (bSectionOpened)
+				PRAGMA_DISABLE_DEPRECATION_WARNINGS
+				const bool bDeprecatedCanBindToContextStruct = Args.OnCanBindToContextStruct.IsBound() && Args.OnCanBindToContextStruct.Execute(ContextStruct.Struct);
+				PRAGMA_ENABLE_DEPRECATION_WARNINGS
+
+				if (bDeprecatedCanBindToContextStruct || (Args.OnCanBindToContextStructWithIndex.IsBound() && Args.OnCanBindToContextStructWithIndex.Execute(ContextStruct.Struct, Index)))
 				{
-					MenuBuilder.EndSection();
+					// If the struct can be be bound to directly, create action for that.
+					MenuBuilder.AddMenuEntry(
+						FUIAction(FExecuteAction::CreateSP(this, &SPropertyBinding::HandleAddBinding, BindingChain)),
+						MakeContextStructWidget(ContextStruct));
+				}
+				if (HasBindableProperties(ContextStruct.Struct, BindingChain))
+				{
+					// Show struct properties.
+					MenuBuilder.AddSubMenu(
+						MakeContextStructWidget(ContextStruct),
+						FNewMenuDelegate::CreateSP(this, &SPropertyBinding::FillPropertyMenu, ContextStruct.Struct, BindingChain));
 				}
-
-				CurrentSection = ContextStruct.Section;
-				MenuBuilder.BeginSection(FName(), CurrentSection);
-				bSectionOpened = true;
-			}
-			
-			// Make first chain element representing the index in the context array.
-			TArray<TSharedPtr<FBindingChainElement>> BindingChain;
-			BindingChain.Emplace(MakeShared<FBindingChainElement>(nullptr, i));
-
-			PRAGMA_DISABLE_DEPRECATION_WARNINGS
-			const bool bDeprecatedCanBindToContextStruct = Args.OnCanBindToContextStruct.IsBound() && Args.OnCanBindToContextStruct.Execute(ContextStruct.Struct);
-			PRAGMA_ENABLE_DEPRECATION_WARNINGS
-
-			if (bDeprecatedCanBindToContextStruct || (Args.OnCanBindToContextStructWithIndex.IsBound() && Args.OnCanBindToContextStructWithIndex.Execute(ContextStruct.Struct, i)))
-			{
-				// If the struct can be be bound to directly, create action for that. 
-				MenuBuilder.AddMenuEntry(
-					FUIAction(FExecuteAction::CreateSP(this, &SPropertyBinding::HandleAddBinding, BindingChain)),
-					MakeContextStructWidget(ContextStruct));
-			}
-			if (HasBindableProperties(ContextStruct.Struct, BindingChain))
-			{
-				// Show struct properties.
-				MenuBuilder.AddSubMenu(
-					MakeContextStructWidget(ContextStruct),
-					FNewMenuDelegate::CreateSP(this, &SPropertyBinding::FillPropertyMenu, ContextStruct.Struct, BindingChain));
 			}
-		}
-
-		if (bSectionOpened)
-		{
 			MenuBuilder.EndSection();
 		}
-
 	}
 
 	FDisplayMetrics DisplayMetrics;
@@ -916,6 +1010,44 @@ void SPropertyBinding::FillPropertyMenu(FMenuBuilder& MenuBuilder, UStruct* InOw
 	}
 }
 
+void SPropertyBinding::FillCategoryMenu(FMenuBuilder& MenuBuilder, const FBindingContextStructCategory* Category)
+{
+	MenuBuilder.BeginSection(NAME_None, Category->Name);
+	for (const FBindingContextStructCategory& SubCategory : Category->SubCategories)
+	{
+		AddCategoryToMenu(MenuBuilder, SubCategory);
+	}
+
+	for(int32 Index : Category->BindingContextStructIndices)
+	{
+		FBindingContextStruct& ContextStruct = BindingContextStructs[Index];
+
+		// Make first chain element representing the index in the context array.
+		TArray<TSharedPtr<FBindingChainElement>> BindingChain;
+		BindingChain.Emplace(MakeShared<FBindingChainElement>(nullptr, Index));
+
+		PRAGMA_DISABLE_DEPRECATION_WARNINGS
+		const bool bDeprecatedCanBindToContextStruct = Args.OnCanBindToContextStruct.IsBound() && Args.OnCanBindToContextStruct.Execute(ContextStruct.Struct);
+		PRAGMA_ENABLE_DEPRECATION_WARNINGS
+
+		if (bDeprecatedCanBindToContextStruct || (Args.OnCanBindToContextStructWithIndex.IsBound() && Args.OnCanBindToContextStructWithIndex.Execute(ContextStruct.Struct, Index)))
+		{
+			// If the struct can be be bound to directly, create action for that.
+			MenuBuilder.AddMenuEntry(
+				FUIAction(FExecuteAction::CreateSP(this, &SPropertyBinding::HandleAddBinding, BindingChain)),
+				MakeContextStructWidget(ContextStruct));
+		}
+		if (HasBindableProperties(ContextStruct.Struct, BindingChain))
+		{
+			// Show struct properties.
+			MenuBuilder.AddSubMenu(
+				MakeContextStructWidget(ContextStruct),
+				FNewMenuDelegate::CreateSP(this, &SPropertyBinding::FillPropertyMenu, ContextStruct.Struct, BindingChain));
+		}
+	}
+	MenuBuilder.EndSection();
+}
+
 const FSlateBrush* SPropertyBinding::GetCurrentBindingImage() const
 {
 	if(Args.CurrentBindingImage.IsSet())
diff --git a/Engine/Plugins/Runtime/PropertyAccess/Source/PropertyAccessEditor/Private/SPropertyBinding.h b/Engine/Plugins/Runtime/PropertyAccess/Source/PropertyAccessEditor/Private/SPropertyBinding.h
--- a/Engine/Plugins/Runtime/PropertyAccess/Source/PropertyAccessEditor/Private/SPropertyBinding.h
+++ b/Engine/Plugins/Runtime/PropertyAccess/Source/PropertyAccessEditor/Private/SPropertyBinding.h
@@ -46,8 +46,17 @@ protected:
 		UFunction* Function;
 	};
 
+	struct FBindingContextStructCategory
+	{
+		FText Name;
+		TArray<FBindingContextStructCategory> SubCategories;
+		TArray<int32> BindingContextStructIndices;
+	};
+
+
 	TSharedRef<SWidget> OnGenerateDelegateMenu();
 	void FillPropertyMenu(FMenuBuilder& MenuBuilder, UStruct* InOwnerStruct, TArray<TSharedPtr<FBindingChainElement>> InBindingChain);
+	void FillCategoryMenu(FMenuBuilder& MenuBuilder, const FBindingContextStructCategory* Category);
 
 	const FSlateBrush* GetCurrentBindingImage() const;
 	FText GetCurrentBindingText() const;
@@ -77,6 +86,7 @@ protected:
 	virtual FReply OnDrop(const FGeometry& MyGeometry, const FDragDropEvent& DragDropEvent) override;
 
 private:
+
 	bool IsClassDenied(UClass* OwnerClass) const;
 	bool IsFieldFromDeniedClass(FFieldVariant Field) const;
 	bool HasBindableProperties(UStruct* InStruct, TArray<TSharedPtr<FBindingChainElement>>& BindingChain) const;
@@ -92,8 +102,16 @@ private:
 	template <typename Predicate>
 	void ForEachBindableFunction(UClass* FromClass, Predicate Pred) const;
 
+	void AddCategoryToMenu(FMenuBuilder& MenuBuilder, const FBindingContextStructCategory& Category);
+	void BuildContextStructCategoryRecursive(TConstArrayView<FString> CategoryNames, TArray<FBindingContextStructCategory>& ParentSubCategories, int32 ContextStructIndex);
+	bool HasCategorySomethingToDisplayRecursive(const FBindingContextStructCategory& Category) const;
+
+	TSharedRef<SWidget> MakeContextStructWidget(const FBindingContextStruct& ContextStruct) const;
+
 	UBlueprint* Blueprint = nullptr;
 	TArray<FBindingContextStruct> BindingContextStructs;
+	// Top level sections of the binding ContextStructs
+	TArray<FBindingContextStructCategory> BindingContextStructSections;
 	FPropertyBindingWidgetArgs Args;
 	FName PropertyName;
 };
diff --git a/Engine/Source/Editor/UnrealEd/Public/IPropertyAccessEditor.h b/Engine/Source/Editor/UnrealEd/Public/IPropertyAccessEditor.h
--- a/Engine/Source/Editor/UnrealEd/Public/IPropertyAccessEditor.h
+++ b/Engine/Source/Editor/UnrealEd/Public/IPropertyAccessEditor.h
@@ -96,6 +96,9 @@ struct FBindingContextStruct
 
 	/** Name of the section to put the struct to. If left empty, no section will be created. */
 	FText Section;
+
+	/** Category separated by | of the struct. Will display as submenus. Not part of the section to avoid converting Text to string and back. */
+	FString Category;
 };
 
 
