PropertyAccessEditor: Add binding text color.

diff --git a/Engine/Plugins/Runtime/PropertyAccess/Source/PropertyAccessEditor/Private/SPropertyBinding.cpp b/Engine/Plugins/Runtime/PropertyAccess/Source/PropertyAccessEditor/Private/SPropertyBinding.cpp
--- a/Engine/Plugins/Runtime/PropertyAccess/Source/PropertyAccessEditor/Private/SPropertyBinding.cpp
+++ b/Engine/Plugins/Runtime/PropertyAccess/Source/PropertyAccessEditor/Private/SPropertyBinding.cpp
@@ -113,7 +113,7 @@ void SPropertyBinding::Construct(const FArguments& InArgs, UBlueprint* InBluepri
 				[
 					SNew(STextBlock)
 					.Text(this, &SPropertyBinding::GetCurrentBindingText)
-					.ColorAndOpacity(FSlateColor::UseForeground())
+					.ColorAndOpacity(this, &SPropertyBinding::GetCurrentBindingTextColor)
 					.Font(IDetailLayoutBuilder::GetDetailFont())
 				]
 			]
@@ -1092,6 +1092,16 @@ FText SPropertyBinding::GetCurrentBindingText() const
 	return LOCTEXT("Bind", "Bind");
 }
 
+FSlateColor SPropertyBinding::GetCurrentBindingTextColor() const
+{
+	if (Args.CurrentBindingTextColor.IsSet())
+	{
+		return Args.CurrentBindingTextColor.Get();
+	}
+
+	return FSlateColor::UseForeground();
+}
+
 FText SPropertyBinding::GetCurrentBindingToolTipText() const
 {
 	if(Args.CurrentBindingToolTipText.IsSet())
diff --git a/Engine/Plugins/Runtime/PropertyAccess/Source/PropertyAccessEditor/Private/SPropertyBinding.h b/Engine/Plugins/Runtime/PropertyAccess/Source/PropertyAccessEditor/Private/SPropertyBinding.h
--- a/Engine/Plugins/Runtime/PropertyAccess/Source/PropertyAccessEditor/Private/SPropertyBinding.h
+++ b/Engine/Plugins/Runtime/PropertyAccess/Source/PropertyAccessEditor/Private/SPropertyBinding.h
@@ -61,6 +61,7 @@ protected:
 	const FSlateBrush* GetLinkIcon() const;
 	const FSlateBrush* GetCurrentBindingImage() const;
 	FText GetCurrentBindingText() const;
+	FSlateColor GetCurrentBindingTextColor() const;
 	FText GetCurrentBindingToolTipText() const;
 	FSlateColor GetCurrentBindingColor() const;
 
diff --git a/Engine/Source/Editor/UnrealEd/Public/IPropertyAccessEditor.h b/Engine/Source/Editor/UnrealEd/Public/IPropertyAccessEditor.h
--- a/Engine/Source/Editor/UnrealEd/Public/IPropertyAccessEditor.h
+++ b/Engine/Source/Editor/UnrealEd/Public/IPropertyAccessEditor.h
@@ -4,11 +4,12 @@
 
 #include "CoreMinimal.h"
 #include "Modules/ModuleInterface.h"
-#include "UObject/Field.h"
 #include "Misc/Attribute.h"
 #include "Features/IModularFeature.h"
-#include "UObject/UnrealType.h"
 #include "Input/Reply.h"
+#include "Styling/SlateColor.h"
+#include "UObject/Field.h"
+#include "UObject/UnrealType.h"
 
 class UBlueprint;
 class IPropertyHandle;
@@ -236,6 +237,9 @@ struct FPropertyBindingWidgetArgs
 
 	/** The current binding's text label */
 	TAttribute<FText> CurrentBindingText;
+	
+	/** The current binding's text label color */
+	TAttribute<FSlateColor> CurrentBindingTextColor;
 
 	/** The current binding's tooltip text label */
 	TAttribute<FText> CurrentBindingToolTipText;
