StateTree: Support binding on array element.

diff --git a/Engine/Plugins/Runtime/PropertyAccess/Source/PropertyAccessEditor/Private/SPropertyBinding.cpp b/Engine/Plugins/Runtime/PropertyAccess/Source/PropertyAccessEditor/Private/SPropertyBinding.cpp
--- a/Engine/Plugins/Runtime/PropertyAccess/Source/PropertyAccessEditor/Private/SPropertyBinding.cpp
+++ b/Engine/Plugins/Runtime/PropertyAccess/Source/PropertyAccessEditor/Private/SPropertyBinding.cpp
@@ -900,10 +900,10 @@ void SPropertyBinding::FillPropertyMenu(FMenuBuilder& MenuBuilder, UStruct* InOw
 	// Property Bindings
 
 	// Only show bindable subobjects and variables if we're generating pure bindings.
-	if ( Args.bGeneratePureBindings && Args.bAllowPropertyBindings )
+	if (Args.bGeneratePureBindings && Args.bAllowPropertyBindings)
 	{
 		FProperty* BindingProperty = nullptr;
-		if(Args.BindableSignature != nullptr)
+		if (Args.BindableSignature != nullptr)
 		{
 			BindingProperty = Args.BindableSignature->GetReturnProperty();
 		}
@@ -930,90 +930,88 @@ void SPropertyBinding::FillPropertyMenu(FMenuBuilder& MenuBuilder, UStruct* InOw
 
 			MenuBuilder.BeginSection("Properties", LOCTEXT("Properties", "Properties"));
 			{
-				ForEachBindableProperty(InOwnerStruct, InBindingChain, [this, &InBindingChain, BindingStruct, &MenuBuilder, &MakeArrayElementPropertyWidget, &MakePropertyWidget, &MakePropertyEntry, &MakeArrayElementEntry] (FProperty* Property)
-				{
-					TArray<TSharedPtr<FBindingChainElement>> NewBindingChain(InBindingChain);
-					NewBindingChain.Emplace(MakeShared<FBindingChainElement>(Property));
-
-					if(CanBindProperty(Property, NewBindingChain))
+				ForEachBindableProperty(InOwnerStruct, InBindingChain, [this, &InBindingChain, BindingStruct, &MenuBuilder, &MakeArrayElementPropertyWidget, &MakePropertyWidget, &MakePropertyEntry, &MakeArrayElementEntry](FProperty* Property)
 					{
-						MakePropertyEntry(Property, NewBindingChain);
-					}
+						TArray<TSharedPtr<FBindingChainElement>> NewBindingChain(InBindingChain);
+						TArray<TSharedPtr<FBindingChainElement>> NewArrayElementBindingChain;
+						NewBindingChain.Emplace(MakeShared<FBindingChainElement>(Property));
 
-					FArrayProperty* ArrayProperty = CastField<FArrayProperty>(Property);
-
-					if(Args.bAllowArrayElementBindings && ArrayProperty != nullptr)
-					{
-						TArray<TSharedPtr<FBindingChainElement>> NewArrayElementBindingChain(InBindingChain);
-						NewArrayElementBindingChain.Emplace(MakeShared<FBindingChainElement>(Property));
-						NewArrayElementBindingChain.Last()->ArrayIndex = 0;
-						if(CanBindProperty(ArrayProperty->Inner, NewArrayElementBindingChain))
+						if (CanBindProperty(Property, NewBindingChain))
 						{
-							MakeArrayElementEntry(ArrayProperty->Inner, NewArrayElementBindingChain);
+							MakePropertyEntry(Property, NewBindingChain);
 						}
-					}
 
-					FProperty* InnerProperty = Property;
-					if(Args.bAllowArrayElementBindings && ArrayProperty != nullptr)
-					{
-						InnerProperty = ArrayProperty->Inner;
-					}
+						FProperty* InnerProperty = Property;
+						FArrayProperty* ArrayProperty = CastField<FArrayProperty>(Property);
 
-					FObjectPropertyBase* ObjectProperty = CastField<FObjectPropertyBase>(InnerProperty);
-					FStructProperty* StructProperty = CastField<FStructProperty>(InnerProperty);
+						const bool bUseArrayProperty = Args.bAllowArrayElementBindings && ArrayProperty != nullptr;
+						if (bUseArrayProperty)
+						{
+							NewArrayElementBindingChain = InBindingChain;
+							NewArrayElementBindingChain.Emplace(MakeShared<FBindingChainElement>(Property));
+							NewArrayElementBindingChain.Last()->ArrayIndex = 0;
+							if (CanBindProperty(ArrayProperty->Inner, NewArrayElementBindingChain))
+							{
+								MakeArrayElementEntry(ArrayProperty->Inner, NewArrayElementBindingChain);
+							}
 
-					UStruct* Struct = nullptr;
-					UClass* Class = nullptr;
+							InnerProperty = ArrayProperty->Inner;
+						}
 
-					if ( ObjectProperty )
-					{
-						Struct = Class = ObjectProperty->PropertyClass;
-					}
-					else if ( StructProperty )
-					{
-						Struct = StructProperty->Struct;
-					}
+						FObjectPropertyBase* ObjectProperty = CastField<FObjectPropertyBase>(InnerProperty);
+						FStructProperty* StructProperty = CastField<FStructProperty>(InnerProperty);
 
-					// Recurse into a struct if it has some properties we can bind to.
-					if (Struct)
-					{
-						if (UStruct* ResolvedStruct = ResolveIndirection(NewBindingChain))
+						UStruct* Struct = nullptr;
+						UClass* Class = nullptr;
+
+						if (ObjectProperty)
 						{
-							Struct = ResolvedStruct;
+							Struct = Class = ObjectProperty->PropertyClass;
 						}
-
-						if (Class)
+						else if (StructProperty)
 						{
-							// Ignore any subobject properties that are not bindable.
-							// Also ignore any class that is explicitly on the deny list.
-							if ( IsClassDenied(Class) || (Args.OnCanBindToSubObjectClass.IsBound() && Args.OnCanBindToSubObjectClass.Execute(Class)))
-							{
-								return;
-							}
+							Struct = StructProperty->Struct;
 						}
 
-						if (HasBindableProperties(Struct, NewBindingChain))
+						// Recurse into a struct if it has some properties we can bind to.
+						if (Struct)
 						{
-							if (Args.bAllowArrayElementBindings && ArrayProperty != nullptr)
+							if (UStruct* ResolvedStruct = ResolveIndirection(NewBindingChain))
 							{
-								TArray<TSharedPtr<FBindingChainElement>> NewArrayElementBindingChain(NewBindingChain);
-								NewArrayElementBindingChain.Emplace(MakeShared<FBindingChainElement>(Property));
-								NewArrayElementBindingChain.Last()->ArrayIndex = 0;
+								Struct = ResolvedStruct;
+							}
 
-								MenuBuilder.AddSubMenu(
-									MakeArrayElementPropertyWidget(ArrayProperty->Inner, NewArrayElementBindingChain),
-									FNewMenuDelegate::CreateSP(this, &SPropertyBinding::FillPropertyMenu, Struct, NewArrayElementBindingChain)
-									);
+							if (Class)
+							{
+								// Ignore any subobject properties that are not bindable.
+								// Also ignore any class that is explicitly on the deny list.
+								if (IsClassDenied(Class) || (Args.OnCanBindToSubObjectClass.IsBound() && Args.OnCanBindToSubObjectClass.Execute(Class)))
+								{
+									return;
+								}
 							}
-							else if (Args.bAllowStructMemberBindings)
+
+							TArray<TSharedPtr<FBindingChainElement>>& BindablePropertyBindingChain = bUseArrayProperty
+								? NewArrayElementBindingChain
+								: NewBindingChain;
+							if (HasBindableProperties(Struct, BindablePropertyBindingChain))
 							{
-								MenuBuilder.AddSubMenu(
-									MakePropertyWidget(Property),
-									FNewMenuDelegate::CreateSP(this, &SPropertyBinding::FillPropertyMenu, Struct, NewBindingChain));
+								if (bUseArrayProperty)
+								{
+									MenuBuilder.AddSubMenu(
+										MakeArrayElementPropertyWidget(ArrayProperty->Inner, NewArrayElementBindingChain),
+										FNewMenuDelegate::CreateSP(this, &SPropertyBinding::FillPropertyMenu, Struct, NewArrayElementBindingChain)
+									);
+								}
+								else if (Args.bAllowStructMemberBindings)
+								{
+									MenuBuilder.AddSubMenu(
+										MakePropertyWidget(Property),
+										FNewMenuDelegate::CreateSP(this, &SPropertyBinding::FillPropertyMenu, Struct, NewBindingChain));
+								}
 							}
 						}
-					}
-				});
+					});
 			}
 			MenuBuilder.EndSection(); //Properties
 		}
diff --git a/Engine/Plugins/Runtime/PropertyBindingUtils/Source/PropertyBindingUtilsEditor/Private/PropertyBindingExtension.cpp b/Engine/Plugins/Runtime/PropertyBindingUtils/Source/PropertyBindingUtilsEditor/Private/PropertyBindingExtension.cpp
--- a/Engine/Plugins/Runtime/PropertyBindingUtils/Source/PropertyBindingUtilsEditor/Private/PropertyBindingExtension.cpp
+++ b/Engine/Plugins/Runtime/PropertyBindingUtils/Source/PropertyBindingUtilsEditor/Private/PropertyBindingExtension.cpp
@@ -1326,7 +1326,7 @@ void FPropertyBindingExtension::ExtendWidgetRow(FDetailWidgetRow& InWidgetRow, c
 
 	Args.BindButtonStyle = &FAppStyle::Get().GetWidgetStyle<FButtonStyle>("HoverHintOnly");
 	Args.bAllowNewBindings = false;
-	Args.bAllowArrayElementBindings = false;
+	Args.bAllowArrayElementBindings = CanBindToArrayElements();
 	Args.bAllowUObjectFunctions = false;
 
 	if (CanPromoteToParameter(InPropertyHandle))
@@ -1438,7 +1438,7 @@ void FPropertyBindingExtension::ExtendWidgetRow(FDetailWidgetRow& InWidgetRow, c
 				}),
 			false);
 	}
-	
+
 	InWidgetRow.ExtensionContent()
 	[
 		PropertyAccessEditor.MakePropertyBindingWidget(BindingContextStructs, Args)
diff --git a/Engine/Plugins/Runtime/PropertyBindingUtils/Source/PropertyBindingUtilsEditor/Public/PropertyBindingExtension.h b/Engine/Plugins/Runtime/PropertyBindingUtils/Source/PropertyBindingUtilsEditor/Public/PropertyBindingExtension.h
--- a/Engine/Plugins/Runtime/PropertyBindingUtils/Source/PropertyBindingUtilsEditor/Public/PropertyBindingExtension.h
+++ b/Engine/Plugins/Runtime/PropertyBindingUtils/Source/PropertyBindingUtilsEditor/Public/PropertyBindingExtension.h
@@ -232,6 +232,11 @@ protected:
 	PROPERTYBINDINGUTILSEDITOR_API virtual bool CanBindToProperty(const FPropertyBindingPath& InTargetPath, const IPropertyHandle& InPropertyHandle) const;
 	PROPERTYBINDINGUTILSEDITOR_API virtual bool GetPromotionToParameterOverrideInternal(const FProperty& InProperty, bool& bOutOverride) const;
 	PROPERTYBINDINGUTILSEDITOR_API virtual void UpdateContextStruct(TConstStructView<FPropertyBindingBindableStructDescriptor> InStructDesc, FBindingContextStruct& InOutContextStruct, TMap<FString, FText>& InOutSectionNames) const;
+	virtual bool CanBindToArrayElements() const
+	{
+		return false;
+	}
+
 	PROPERTYBINDINGUTILSEDITOR_API static TSharedRef<SWidget> MakeContextStructWidget(const FPropertyBindingBindableStructDescriptor& InContextStruct);
 
 private:
diff --git a/Engine/Plugins/Runtime/StateTree/Source/StateTreeEditorModule/Private/Customizations/StateTreeBindingExtension.h b/Engine/Plugins/Runtime/StateTree/Source/StateTreeEditorModule/Private/Customizations/StateTreeBindingExtension.h
--- a/Engine/Plugins/Runtime/StateTree/Source/StateTreeEditorModule/Private/Customizations/StateTreeBindingExtension.h
+++ b/Engine/Plugins/Runtime/StateTree/Source/StateTreeEditorModule/Private/Customizations/StateTreeBindingExtension.h
@@ -27,6 +27,10 @@ private:
 	virtual bool CanBindToProperty(const FPropertyBindingPath& InTargetPath, const IPropertyHandle& InPropertyHandle) const override;
 	virtual bool GetPromotionToParameterOverrideInternal(const FProperty& InProperty, bool& bOutOverride) const override;
 	virtual void UpdateContextStruct(TConstStructView<FPropertyBindingBindableStructDescriptor> InStructDesc,FBindingContextStruct& InOutContextStruct, TMap<FString, FText>& InOutSectionNames) const override;
+	virtual bool CanBindToArrayElements() const override
+	{
+		return true;
+	}
 };
 
 /* Overrides bound property's children composition. */
