Expose function for generating property access menus to property access editor

diff --git a/Engine/Plugins/Runtime/PropertyAccess/Source/PropertyAccessEditor/Private/PropertyAccessEditorModule.cpp b/Engine/Plugins/Runtime/PropertyAccess/Source/PropertyAccessEditor/Private/PropertyAccessEditorModule.cpp
--- a/Engine/Plugins/Runtime/PropertyAccess/Source/PropertyAccessEditor/Private/PropertyAccessEditorModule.cpp
+++ b/Engine/Plugins/Runtime/PropertyAccess/Source/PropertyAccessEditor/Private/PropertyAccessEditorModule.cpp
@@ -27,6 +27,12 @@ public:
 			.Args(InArgs);
 	}
 
+	
+	virtual void FillPropertyMenu(FMenuBuilder& MenuBuilder, const FPropertyBindingWidgetArgs& InArgs, FName InPropertyName, UStruct* InOwnerStruct, TArray<TSharedPtr<FBindingChainElement>> InBindingChain) const override
+	{
+		SPropertyBinding::FillPropertyMenu(MenuBuilder, InArgs, nullptr, InPropertyName, InOwnerStruct, InBindingChain);
+	}
+	
 	virtual TSharedRef<SWidget> MakePropertyBindingWidget(const TArray<FBindingContextStruct>& InBindingContextStructs, const FPropertyBindingWidgetArgs& InArgs) const override
 	{
 		return SNew(SPropertyBinding, nullptr, InBindingContextStructs)
diff --git a/Engine/Plugins/Runtime/PropertyAccess/Source/PropertyAccessEditor/Private/SPropertyBinding.cpp b/Engine/Plugins/Runtime/PropertyAccess/Source/PropertyAccessEditor/Private/SPropertyBinding.cpp
--- a/Engine/Plugins/Runtime/PropertyAccess/Source/PropertyAccessEditor/Private/SPropertyBinding.cpp
+++ b/Engine/Plugins/Runtime/PropertyAccess/Source/PropertyAccessEditor/Private/SPropertyBinding.cpp
@@ -136,7 +136,7 @@ void SPropertyBinding::Construct(const FArguments& InArgs, UBlueprint* InBluepri
 	];
 }
 
-bool SPropertyBinding::IsClassDenied(UClass* InClass) const
+bool SPropertyBinding::IsClassDenied(const FPropertyBindingWidgetArgs& Args, UClass* InClass)
 {
 	if(Args.OnCanBindToClass.IsBound() && Args.OnCanBindToClass.Execute(InClass))
 	{
@@ -146,17 +146,17 @@ bool SPropertyBinding::IsClassDenied(UClass* InClass) const
 	return true;
 }
 
-bool SPropertyBinding::IsFieldFromDeniedClass(FFieldVariant Field) const
+bool SPropertyBinding::IsFieldFromDeniedClass(const FPropertyBindingWidgetArgs& Args, FFieldVariant Field)
 {
-	return IsClassDenied(Field.GetOwnerClass());
+	return IsClassDenied(Args, Field.GetOwnerClass());
 }
 
 template <typename Predicate>
-void SPropertyBinding::ForEachBindableFunction(UClass* FromClass, Predicate Pred) const
+void SPropertyBinding::ForEachBindableFunction(const FPropertyBindingWidgetArgs& Args, UClass* FromClass, Predicate Pred)
 {
 	if(Args.OnCanBindFunction.IsBound())
 	{
-		auto CheckBindableClass = [this, &Pred](UClass* InBindableClass)
+		auto CheckBindableClass = [&Args, &Pred](UClass* InBindableClass)
 		{
 			// Early exit if this class is excluded by permissions
 			if ( !FBlueprintActionDatabase::IsClassAllowed(InBindableClass, FBlueprintActionDatabase::EPermissionsContext::Property) )
@@ -170,7 +170,7 @@ void SPropertyBinding::ForEachBindableFunction(UClass* FromClass, Predicate Pred
 				UFunction* Function = *FuncIt;
 
 				// Stop processing functions after reaching a base class that it doesn't make sense to go beyond.
-				if ( IsFieldFromDeniedClass(Function) )
+				if ( IsFieldFromDeniedClass(Args, Function) )
 				{
 					break;
 				}
@@ -317,7 +317,7 @@ bool SPropertyBinding::HasCategorySomethingToDisplayRecursive(const FBindingCont
 
 		TArray<TSharedPtr<FBindingChainElement>> BindingChain;
 		BindingChain.Emplace(MakeShared<FBindingChainElement>(nullptr, ContextStructIndex));
-		if (HasBindableProperties(Struct, BindingChain))
+		if (HasBindableProperties(Args, Blueprint, Struct, BindingChain))
 		{
 			return true;
 		}
@@ -382,7 +382,7 @@ TSharedRef<SWidget> SPropertyBinding::MakeContextStructWidget(const FBindingCont
 }
 
 template <typename Predicate>
-void SPropertyBinding::ForEachBindableProperty(UStruct* InStruct, const TArray<TSharedPtr<FBindingChainElement>>& BindingChain, Predicate Pred) const
+void SPropertyBinding::ForEachBindableProperty(const FPropertyBindingWidgetArgs& Args,  const UBlueprint* Blueprint, UStruct* InStruct, const TArray<TSharedPtr<FBindingChainElement>>& BindingChain, Predicate Pred)
 {
 	PRAGMA_DISABLE_DEPRECATION_WARNINGS
 	const bool bIsOnCanBindPropertyBound = Args.OnCanBindProperty.IsBound() || Args.OnCanBindPropertyWithBindingChain.IsBound();
@@ -397,12 +397,12 @@ void SPropertyBinding::ForEachBindableProperty(UStruct* InStruct, const TArray<T
 			FProperty* Property = *PropIt;
 
 			// Stop processing properties after reaching the stopped base class
-			if ( IsFieldFromDeniedClass(Property) )
+			if ( IsFieldFromDeniedClass(Args, Property) )
 			{
 				break;
 			}
 
-			if (!CanAcceptPropertyOrChildren(Property, BindingChain))
+			if (!CanAcceptPropertyOrChildren(Args, Property, BindingChain))
 			{
 				continue;
 			}
@@ -431,13 +431,13 @@ void SPropertyBinding::ForEachBindableProperty(UStruct* InStruct, const TArray<T
 	}
 }
 
-bool SPropertyBinding::HasBindableProperties(UStruct* InStruct, TArray<TSharedPtr<FBindingChainElement>>& BindingChain) const
+bool SPropertyBinding::HasBindableProperties(const FPropertyBindingWidgetArgs& Args, const UBlueprint* Blueprint, UStruct* InStruct, TArray<TSharedPtr<FBindingChainElement>>& BindingChain)
 {
 	TSet<UStruct*> VisitedStructs;
-	return HasBindablePropertiesRecursive(InStruct, VisitedStructs, BindingChain);
+	return HasBindablePropertiesRecursive(Args, Blueprint, InStruct, VisitedStructs, BindingChain);
 }
 
-bool SPropertyBinding::HasBindablePropertiesRecursive(UStruct* InStruct, TSet<UStruct*>& VisitedStructs, TArray<TSharedPtr<FBindingChainElement>>& BindingChain) const
+bool SPropertyBinding::HasBindablePropertiesRecursive(const FPropertyBindingWidgetArgs& Args, const UBlueprint* Blueprint, UStruct* InStruct, TSet<UStruct*>& VisitedStructs, TArray<TSharedPtr<FBindingChainElement>>& BindingChain)
 {
 	if (VisitedStructs.Contains(InStruct))
 	{
@@ -454,14 +454,14 @@ bool SPropertyBinding::HasBindablePropertiesRecursive(UStruct* InStruct, TSet<US
 	}
 	
 	int32 BindableCount = 0;
-	ForEachBindableProperty(InStruct, BindingChain, [this, &BindableCount, &VisitedStructs, &BindingChain] (FProperty* Property)
+	ForEachBindableProperty(Args, Blueprint, InStruct, BindingChain, [&Args, Blueprint, &BindableCount, &VisitedStructs, &BindingChain] (FProperty* Property)
 	{
 		BindingChain.Emplace(MakeShared<FBindingChainElement>(Property));
 		ON_SCOPE_EXIT{ BindingChain.Pop(); };
 
 		FArrayProperty* ArrayProperty = CastField<FArrayProperty>(Property);
 
-		if(CanBindProperty(Property, BindingChain))
+		if(CanBindProperty(Args, Property, BindingChain))
 		{
 			BindableCount++;
 			return;
@@ -473,7 +473,7 @@ bool SPropertyBinding::HasBindablePropertiesRecursive(UStruct* InStruct, TSet<US
 			BindingChain.Last()->ArrayIndex = 0;
 			ON_SCOPE_EXIT{ BindingChain.Pop(); };
 
-			if(CanBindProperty(ArrayProperty->Inner, BindingChain))
+			if(CanBindProperty(Args, ArrayProperty->Inner, BindingChain))
 			{
 				BindableCount++;
 				return;
@@ -503,19 +503,19 @@ bool SPropertyBinding::HasBindablePropertiesRecursive(UStruct* InStruct, TSet<US
 
 		if (Struct)
 		{
-			if (UStruct* ResolvedStruct = ResolveIndirection(BindingChain))
+			if (UStruct* ResolvedStruct = ResolveIndirection(Args, BindingChain))
 			{
 				Struct = ResolvedStruct;
 			}
 
 			// Recurse into a struct, except if it is the same type as the one we're binding.
-			if (HasBindablePropertiesRecursive(Struct, VisitedStructs, BindingChain))
+			if (HasBindablePropertiesRecursive(Args, Blueprint, Struct, VisitedStructs, BindingChain))
 			{
 				if (Class)
 				{
 					// Ignore any subobject properties that are not bindable.
 					// Also ignore any class that is explicitly on the deny list.
-					if (IsClassDenied(Class) || (Args.OnCanBindToSubObjectClass.IsBound() && Args.OnCanBindToSubObjectClass.Execute(Class)))
+					if (IsClassDenied(Args, Class) || (Args.OnCanBindToSubObjectClass.IsBound() && Args.OnCanBindToSubObjectClass.Execute(Class)))
 					{
 						return;
 					}
@@ -535,7 +535,7 @@ bool SPropertyBinding::HasBindablePropertiesRecursive(UStruct* InStruct, TSet<US
 
 	if(UClass* Class = Cast<UClass>(InStruct))
 	{
-		ForEachBindableFunction(Class, [this, &BindableCount, &VisitedStructs, &BindingChain](const FFunctionInfo& Info)
+		ForEachBindableFunction(Args, Class, [&Args, &Blueprint, &BindableCount, &VisitedStructs, &BindingChain](const FFunctionInfo& Info)
 		{
 			BindingChain.Emplace(MakeShared<FBindingChainElement>(Info.Function));
 			ON_SCOPE_EXIT{ BindingChain.Pop(); };
@@ -543,7 +543,7 @@ bool SPropertyBinding::HasBindablePropertiesRecursive(UStruct* InStruct, TSet<US
 			FProperty* ReturnProperty = Info.Function->GetReturnProperty();
 			
 			// We can get here if we accept non-leaf UObject functions, so if so we need to check the return value for compatibility
-			if(!Args.bAllowUObjectFunctions || CanBindProperty(ReturnProperty, BindingChain))
+			if(!Args.bAllowUObjectFunctions || CanBindProperty(Args, ReturnProperty, BindingChain))
 			{
 				BindableCount++;
 			}
@@ -553,7 +553,7 @@ bool SPropertyBinding::HasBindablePropertiesRecursive(UStruct* InStruct, TSet<US
 			{
 				if(FObjectPropertyBase* ObjectPropertyBase = CastField<FObjectPropertyBase>(ReturnProperty))
 				{
-					if (HasBindablePropertiesRecursive(ObjectPropertyBase->PropertyClass, VisitedStructs, BindingChain))
+					if (HasBindablePropertiesRecursive(Args, Blueprint,ObjectPropertyBase->PropertyClass, VisitedStructs, BindingChain))
 					{
 						if (Args.bAllowUObjectFunctions)
 						{
@@ -563,7 +563,7 @@ bool SPropertyBinding::HasBindablePropertiesRecursive(UStruct* InStruct, TSet<US
 				}
 				else if(FStructProperty* StructProperty = CastField<FStructProperty>(ReturnProperty))
 				{
-					if (HasBindablePropertiesRecursive(StructProperty->Struct, VisitedStructs, BindingChain))
+					if (HasBindablePropertiesRecursive(Args, Blueprint, StructProperty->Struct, VisitedStructs, BindingChain))
 					{
 						BindableCount++;
 					}
@@ -657,7 +657,7 @@ TSharedRef<SWidget> SPropertyBinding::OnGenerateDelegateMenu()
 						FUIAction(FExecuteAction::CreateSP(this, &SPropertyBinding::HandleAddBinding, BindingChain)),
 						MakeContextStructWidget(ContextStruct));
 				}
-				if (HasBindableProperties(ContextStruct.Struct, BindingChain))
+				if (HasBindableProperties(Args, Blueprint, ContextStruct.Struct, BindingChain))
 				{
 					// Show struct properties.
 					MenuBuilder.AddSubMenu(
@@ -684,7 +684,12 @@ TSharedRef<SWidget> SPropertyBinding::OnGenerateDelegateMenu()
 
 void SPropertyBinding::FillPropertyMenu(FMenuBuilder& MenuBuilder, UStruct* InOwnerStruct, TArray<TSharedPtr<FBindingChainElement>> InBindingChain)
 {
-	auto MakeArrayElementPropertyWidget = [this](FProperty* InProperty, const TArray<TSharedPtr<FBindingChainElement>>& InNewBindingChain)
+	FillPropertyMenu(MenuBuilder, Args, Blueprint, PropertyName, InOwnerStruct, InBindingChain);	
+}
+
+void SPropertyBinding::FillPropertyMenu(FMenuBuilder& MenuBuilder, const FPropertyBindingWidgetArgs& Args, const UBlueprint* Blueprint, FName PropertyName,  UStruct* InOwnerStruct, TArray<TSharedPtr<FBindingChainElement>> InBindingChain)
+{
+	auto MakeArrayElementPropertyWidget = [&Args, PropertyName](FProperty* InProperty, const TArray<TSharedPtr<FBindingChainElement>>& InNewBindingChain)
 	{
 		const UEdGraphSchema_K2* Schema = GetDefault<UEdGraphSchema_K2>();
 
@@ -729,15 +734,21 @@ void SPropertyBinding::FillPropertyMenu(FMenuBuilder& MenuBuilder, UStruct* InOw
 					SNew(SNumericEntryBox<int32>)
 					.ToolTipText(LOCTEXT("ArrayIndex", "Array Index"))
 					.Value_Lambda([LeafElement](){ return LeafElement->ArrayIndex == INDEX_NONE ? TOptional<int32>() : TOptional<int32>(LeafElement->ArrayIndex); })
-					.OnValueCommitted(this, &SPropertyBinding::HandleSetBindingArrayIndex, InProperty, InNewBindingChain)
+					.OnValueCommitted_Lambda([Args, PropertyName, InProperty, InNewBindingChain](int32 ArrayIndex, ETextCommit::Type TextCommitType)
+						{
+							HandleSetBindingArrayIndex(Args, PropertyName, ArrayIndex, TextCommitType, InProperty, InNewBindingChain);
+						})
 				]
 			];
 	};
 
-	auto MakeArrayElementEntry = [this, &MenuBuilder, &MakeArrayElementPropertyWidget](FProperty* InProperty, const TArray<TSharedPtr<FBindingChainElement>>& InNewBindingChain)
+	auto MakeArrayElementEntry = [&Args, PropertyName, &MenuBuilder, &MakeArrayElementPropertyWidget](FProperty* InProperty, const TArray<TSharedPtr<FBindingChainElement>>& InNewBindingChain)
 	{
 		MenuBuilder.AddMenuEntry(
-			FUIAction(FExecuteAction::CreateSP(this, &SPropertyBinding::HandleAddBinding, InNewBindingChain)),
+			FUIAction(FExecuteAction::CreateLambda([Args, PropertyName, InNewBindingChain]()
+				{
+					HandleAddBinding(Args, PropertyName, InNewBindingChain);
+				})),
 			SNew(SHorizontalBox)
 			+SHorizontalBox::Slot()
 			.FillWidth(1.0f)
@@ -752,7 +763,7 @@ void SPropertyBinding::FillPropertyMenu(FMenuBuilder& MenuBuilder, UStruct* InOw
 			]);
 	};
 
-	auto MakePropertyWidget = [this](FProperty* InProperty)
+	auto MakePropertyWidget = [](FProperty* InProperty)
 	{
 		const UEdGraphSchema_K2* Schema = GetDefault<UEdGraphSchema_K2>();
 
@@ -803,14 +814,18 @@ void SPropertyBinding::FillPropertyMenu(FMenuBuilder& MenuBuilder, UStruct* InOw
 			];
 	};
 
-	auto MakePropertyEntry = [this, &MenuBuilder, &MakePropertyWidget](FProperty* InProperty, const TArray<TSharedPtr<FBindingChainElement>>& InNewBindingChain)
+	auto MakePropertyEntry = [&Args, PropertyName, &MenuBuilder, &MakePropertyWidget](FProperty* InProperty, const TArray<TSharedPtr<FBindingChainElement>>& InNewBindingChain)
 	{
 		MenuBuilder.AddMenuEntry(
-			FUIAction(FExecuteAction::CreateSP(this, &SPropertyBinding::HandleAddBinding, InNewBindingChain)),
+			FUIAction(FExecuteAction::CreateLambda(
+				[Args, PropertyName, InNewBindingChain]()
+				{
+					HandleAddBinding(Args, PropertyName, InNewBindingChain);
+				})),
 			MakePropertyWidget(InProperty));
 	};
 
-	auto MakeFunctionWidget = [this](const FFunctionInfo& Info)
+	auto MakeFunctionWidget = [](const FFunctionInfo& Info)
 	{
 		static FName FunctionIcon(TEXT("GraphEditor.Function_16x"));
 
@@ -860,17 +875,21 @@ void SPropertyBinding::FillPropertyMenu(FMenuBuilder& MenuBuilder, UStruct* InOw
 
 			MenuBuilder.BeginSection("Functions", LOCTEXT("Functions", "Functions"));
 			{
-				ForEachBindableFunction(OwnerClass, [this, &InBindingChain, &MenuBuilder, &MakeFunctionWidget] (const FFunctionInfo& Info) 
+				ForEachBindableFunction(Args, OwnerClass, [&Args, Blueprint, &PropertyName, &InBindingChain, &MenuBuilder, &MakeFunctionWidget] (const FFunctionInfo& Info) 
 				{
 					TArray<TSharedPtr<FBindingChainElement>> NewBindingChain(InBindingChain);
 					NewBindingChain.Emplace(MakeShared<FBindingChainElement>(Info.Function));
 
 					FProperty* ReturnProperty = Info.Function->GetReturnProperty();
 					// We can get here if we accept non-leaf UObject functions, so if so we need to check the return value for compatibility
-					if(!Args.bAllowUObjectFunctions || CanBindProperty(ReturnProperty, InBindingChain))
+					if(!Args.bAllowUObjectFunctions || CanBindProperty(Args, ReturnProperty, InBindingChain))
 					{
 						MenuBuilder.AddMenuEntry(
-							FUIAction(FExecuteAction::CreateSP(this, &SPropertyBinding::HandleAddBinding, NewBindingChain)),
+							FUIAction(FExecuteAction::CreateLambda(
+								[Args, PropertyName, NewBindingChain]()
+									{
+										HandleAddBinding(Args, PropertyName, NewBindingChain);
+									})),
 							MakeFunctionWidget(Info));
 					}
 
@@ -879,20 +898,27 @@ void SPropertyBinding::FillPropertyMenu(FMenuBuilder& MenuBuilder, UStruct* InOw
 					{
 						if(FObjectPropertyBase* ObjectPropertyBase = CastField<FObjectPropertyBase>(ReturnProperty))
 						{
-							if (HasBindableProperties(ObjectPropertyBase->PropertyClass, NewBindingChain))
+							if (HasBindableProperties(Args, Blueprint, ObjectPropertyBase->PropertyClass,  NewBindingChain))
 							{
 								MenuBuilder.AddSubMenu(
 									MakeFunctionWidget(Info),
-									FNewMenuDelegate::CreateSP(this, &SPropertyBinding::FillPropertyMenu, static_cast<UStruct*>(ObjectPropertyBase->PropertyClass), NewBindingChain));
+									FNewMenuDelegate::CreateLambda(
+										[Args, Blueprint, &PropertyName, ObjectPropertyBase, NewBindingChain](FMenuBuilder& MenuBuilder)
+											{
+												FillPropertyMenu(MenuBuilder, Args, Blueprint, PropertyName, static_cast<UStruct*>(ObjectPropertyBase->PropertyClass), NewBindingChain);
+											}));
 							}
 						}
 						else if(FStructProperty* StructProperty = CastField<FStructProperty>(ReturnProperty))
 						{
-							if (HasBindableProperties(StructProperty->Struct, NewBindingChain))
+							if (HasBindableProperties(Args, Blueprint, StructProperty->Struct, NewBindingChain))
 							{
 								MenuBuilder.AddSubMenu(
 									MakeFunctionWidget(Info),
-									FNewMenuDelegate::CreateSP(this, &SPropertyBinding::FillPropertyMenu, static_cast<UStruct*>(StructProperty->Struct), NewBindingChain));
+									FNewMenuDelegate::CreateLambda([Args, Blueprint, &PropertyName, StructProperty, NewBindingChain](FMenuBuilder& MenuBuilder)
+								{
+									FillPropertyMenu(MenuBuilder, Args, Blueprint, PropertyName, static_cast<UStruct*>(StructProperty->Struct), NewBindingChain);
+								}));
 							}
 						}
 					}
@@ -920,7 +946,7 @@ void SPropertyBinding::FillPropertyMenu(FMenuBuilder& MenuBuilder, UStruct* InOw
 
 		// Find the binder that can handle the delegate return type, don't bother allowing people 
 		// to look for bindings that we don't support
-		if (CanBindProperty(BindingProperty, InBindingChain))
+		if (CanBindProperty(Args, BindingProperty, InBindingChain))
 		{
 			UStruct* BindingStruct = nullptr;
 			FObjectPropertyBase* ObjectProperty = CastField<FObjectPropertyBase>(BindingProperty);
@@ -936,13 +962,13 @@ void SPropertyBinding::FillPropertyMenu(FMenuBuilder& MenuBuilder, UStruct* InOw
 
 			MenuBuilder.BeginSection("Properties", LOCTEXT("Properties", "Properties"));
 			{
-				ForEachBindableProperty(InOwnerStruct, InBindingChain, [this, &InBindingChain, BindingStruct, &MenuBuilder, &MakeArrayElementPropertyWidget, &MakePropertyWidget, &MakePropertyEntry, &MakeArrayElementEntry](FProperty* Property)
+				ForEachBindableProperty(Args, Blueprint, InOwnerStruct, InBindingChain, [&Args, Blueprint, &PropertyName, &InBindingChain, BindingStruct, &MenuBuilder, &MakeArrayElementPropertyWidget, &MakePropertyWidget, &MakePropertyEntry, &MakeArrayElementEntry](FProperty* Property)
 					{
 						TArray<TSharedPtr<FBindingChainElement>> NewBindingChain(InBindingChain);
 						TArray<TSharedPtr<FBindingChainElement>> NewArrayElementBindingChain;
 						NewBindingChain.Emplace(MakeShared<FBindingChainElement>(Property));
 
-						if (CanBindProperty(Property, NewBindingChain))
+						if (CanBindProperty(Args, Property, NewBindingChain))
 						{
 							MakePropertyEntry(Property, NewBindingChain);
 						}
@@ -956,7 +982,7 @@ void SPropertyBinding::FillPropertyMenu(FMenuBuilder& MenuBuilder, UStruct* InOw
 							NewArrayElementBindingChain = InBindingChain;
 							NewArrayElementBindingChain.Emplace(MakeShared<FBindingChainElement>(Property));
 							NewArrayElementBindingChain.Last()->ArrayIndex = 0;
-							if (CanBindProperty(ArrayProperty->Inner, NewArrayElementBindingChain))
+							if (CanBindProperty(Args, ArrayProperty->Inner, NewArrayElementBindingChain))
 							{
 								MakeArrayElementEntry(ArrayProperty->Inner, NewArrayElementBindingChain);
 							}
@@ -982,7 +1008,7 @@ void SPropertyBinding::FillPropertyMenu(FMenuBuilder& MenuBuilder, UStruct* InOw
 						// Recurse into a struct if it has some properties we can bind to.
 						if (Struct)
 						{
-							if (UStruct* ResolvedStruct = ResolveIndirection(NewBindingChain))
+							if (UStruct* ResolvedStruct = ResolveIndirection(Args, NewBindingChain))
 							{
 								Struct = ResolvedStruct;
 							}
@@ -991,7 +1017,7 @@ void SPropertyBinding::FillPropertyMenu(FMenuBuilder& MenuBuilder, UStruct* InOw
 							{
 								// Ignore any subobject properties that are not bindable.
 								// Also ignore any class that is explicitly on the deny list.
-								if (IsClassDenied(Class) || (Args.OnCanBindToSubObjectClass.IsBound() && Args.OnCanBindToSubObjectClass.Execute(Class)))
+								if (IsClassDenied(Args, Class) || (Args.OnCanBindToSubObjectClass.IsBound() && Args.OnCanBindToSubObjectClass.Execute(Class)))
 								{
 									return;
 								}
@@ -1000,20 +1026,29 @@ void SPropertyBinding::FillPropertyMenu(FMenuBuilder& MenuBuilder, UStruct* InOw
 							TArray<TSharedPtr<FBindingChainElement>>& BindablePropertyBindingChain = bUseArrayProperty
 								? NewArrayElementBindingChain
 								: NewBindingChain;
-							if (HasBindableProperties(Struct, BindablePropertyBindingChain))
+							if (HasBindableProperties(Args, Blueprint, Struct, BindablePropertyBindingChain))
 							{
 								if (bUseArrayProperty)
 								{
 									MenuBuilder.AddSubMenu(
 										MakeArrayElementPropertyWidget(ArrayProperty->Inner, NewArrayElementBindingChain),
-										FNewMenuDelegate::CreateSP(this, &SPropertyBinding::FillPropertyMenu, Struct, NewArrayElementBindingChain)
+										FNewMenuDelegate::CreateLambda(
+											[Args, Blueprint, PropertyName, Struct, NewArrayElementBindingChain](FMenuBuilder& MenuBuilder)
+												{
+													FillPropertyMenu(MenuBuilder, Args, Blueprint, PropertyName, Struct, NewArrayElementBindingChain);
+												})
 									);
 								}
 								else if (Args.bAllowStructMemberBindings)
 								{
 									MenuBuilder.AddSubMenu(
 										MakePropertyWidget(Property),
-										FNewMenuDelegate::CreateSP(this, &SPropertyBinding::FillPropertyMenu, Struct, NewBindingChain));
+										FNewMenuDelegate::CreateLambda(
+											[Args, Blueprint, PropertyName, Struct, NewBindingChain](FMenuBuilder& MenuBuilder)
+												{
+													FillPropertyMenu(MenuBuilder, Args, Blueprint, PropertyName, Struct, NewBindingChain);
+												})
+									);
 								}
 							}
 						}
@@ -1059,7 +1094,7 @@ void SPropertyBinding::FillCategoryMenu(FMenuBuilder& MenuBuilder, const FBindin
 				FUIAction(FExecuteAction::CreateSP(this, &SPropertyBinding::HandleAddBinding, BindingChain)),
 				MakeContextStructWidget(ContextStruct));
 		}
-		if (HasBindableProperties(ContextStruct.Struct, BindingChain))
+		if (HasBindableProperties(Args, Blueprint, ContextStruct.Struct, BindingChain))
 		{
 			// Show struct properties.
 			MenuBuilder.AddSubMenu(
@@ -1156,7 +1191,12 @@ void SPropertyBinding::HandleRemoveBinding()
 	}
 }
 
-void SPropertyBinding::HandleAddBinding(TArray<TSharedPtr<FBindingChainElement>> InBindingChain)
+void SPropertyBinding::HandleAddBinding(const TArray<TSharedPtr<FBindingChainElement>> InBindingChain) const
+{
+	HandleAddBinding(Args, PropertyName, InBindingChain);
+}
+
+void SPropertyBinding::HandleAddBinding(const FPropertyBindingWidgetArgs& Args, FName PropertyName, const TArray<TSharedPtr<FBindingChainElement>>& InBindingChain)
 {
 	if(Args.OnAddBinding.IsBound())
 	{
@@ -1172,14 +1212,14 @@ void SPropertyBinding::HandleAddBinding(TArray<TSharedPtr<FBindingChainElement>>
 	}
 }
 
-void SPropertyBinding::HandleSetBindingArrayIndex(int32 InArrayIndex, ETextCommit::Type InCommitType, FProperty* InProperty, TArray<TSharedPtr<FBindingChainElement>> InBindingChain)
+void SPropertyBinding::HandleSetBindingArrayIndex(const FPropertyBindingWidgetArgs& Args, FName PropertyName, int32 InArrayIndex, ETextCommit::Type InCommitType, FProperty* InProperty, TArray<TSharedPtr<FBindingChainElement>> InBindingChain)
 {
 	InBindingChain.Last()->ArrayIndex = InArrayIndex;
 
 	// If the user hit enter on a compatible property, assume they want to accept
-	if(CanBindProperty(InProperty, InBindingChain) && InCommitType == ETextCommit::OnEnter)
+	if(CanBindProperty(Args, InProperty, InBindingChain) && InCommitType == ETextCommit::OnEnter)
 	{
-		HandleAddBinding(InBindingChain);
+		HandleAddBinding(Args, PropertyName, InBindingChain);
 	}
 }
 
@@ -1226,7 +1266,7 @@ void SPropertyBinding::HandleCreateAndAddBinding()
 	TArray<TSharedPtr<FBindingChainElement>> BindingPath;
 	BindingPath.Emplace(MakeShared<FBindingChainElement>(Function));
 
-	HandleAddBinding(BindingPath);
+	HandleAddBinding(Args, PropertyName, BindingPath);
 
 	// Only mark bindings as pure that need to be.
 	if ( Args.bGeneratePureBindings )
@@ -1238,7 +1278,7 @@ void SPropertyBinding::HandleCreateAndAddBinding()
 	HandleGotoBindingClicked();
 }
 
-UStruct* SPropertyBinding::ResolveIndirection(const TArray<TSharedPtr<FBindingChainElement>>& InBindingChain) const
+UStruct* SPropertyBinding::ResolveIndirection(const FPropertyBindingWidgetArgs& Args, const TArray<TSharedPtr<FBindingChainElement>>& InBindingChain)
 {
 	UStruct* ResolvedStruct = nullptr;
 
@@ -1283,7 +1323,7 @@ FReply SPropertyBinding::HandleGotoBindingClicked()
 	return FReply::Unhandled();
 }
 
-bool SPropertyBinding::CanAcceptPropertyOrChildren(FProperty* InProperty, TConstArrayView<TSharedPtr<FBindingChainElement>> InBindingChain) const
+bool SPropertyBinding::CanAcceptPropertyOrChildren(const FPropertyBindingWidgetArgs& Args, FProperty* InProperty, TConstArrayView<TSharedPtr<FBindingChainElement>> InBindingChain)
 {
 	PRAGMA_DISABLE_DEPRECATION_WARNINGS
 	if (Args.OnCanAcceptPropertyOrChildren.IsBound() && Args.OnCanAcceptPropertyOrChildren.Execute(InProperty) == false)
@@ -1309,7 +1349,7 @@ bool SPropertyBinding::CanAcceptPropertyOrChildren(FProperty* InProperty, TConst
 	return true;
 }
 
-bool SPropertyBinding::CanBindProperty(FProperty* InProperty, TConstArrayView<TSharedPtr<FBindingChainElement>> InBindingChain) const
+bool SPropertyBinding::CanBindProperty(const FPropertyBindingWidgetArgs& Args, FProperty* InProperty, TConstArrayView<TSharedPtr<FBindingChainElement>> InBindingChain)
 {
 	PRAGMA_DISABLE_DEPRECATION_WARNINGS
 	if (Args.OnCanBindProperty.IsBound() && Args.OnCanBindProperty.Execute(InProperty) == false)
diff --git a/Engine/Plugins/Runtime/PropertyAccess/Source/PropertyAccessEditor/Private/SPropertyBinding.h b/Engine/Plugins/Runtime/PropertyAccess/Source/PropertyAccessEditor/Private/SPropertyBinding.h
--- a/Engine/Plugins/Runtime/PropertyAccess/Source/PropertyAccessEditor/Private/SPropertyBinding.h
+++ b/Engine/Plugins/Runtime/PropertyAccess/Source/PropertyAccessEditor/Private/SPropertyBinding.h
@@ -24,6 +24,7 @@ public:
 
 	void Construct(const FArguments& InArgs, UBlueprint* InBlueprint, const TArray<FBindingContextStruct>& InBindingContextStructs);
 
+	static void FillPropertyMenu(FMenuBuilder& MenuBuilder,  const FPropertyBindingWidgetArgs& Args, const UBlueprint* Blueprint, FName PropertyName, UStruct* InOwnerStruct, TArray<TSharedPtr<FBindingChainElement>> InBindingChain);
 protected:
 	struct FFunctionInfo
 	{
@@ -68,22 +69,23 @@ protected:
 	bool CanRemoveBinding();
 	void HandleRemoveBinding();
 
-	void HandleAddBinding(TArray<TSharedPtr<FBindingChainElement>> InBindingChain);
-	void HandleSetBindingArrayIndex(int32 InArrayIndex, ETextCommit::Type InCommitType, FProperty* InProperty, TArray<TSharedPtr<FBindingChainElement>> InBindingChain);
+	void HandleAddBinding(const TArray<TSharedPtr<FBindingChainElement>> InBindingChain) const;
+	static void HandleAddBinding(const FPropertyBindingWidgetArgs& Args, FName PropertyName, const TArray<TSharedPtr<FBindingChainElement>>& InBindingChain);
+	static void HandleSetBindingArrayIndex(const FPropertyBindingWidgetArgs& Args, FName PropertyName, int32 InArrayIndex, ETextCommit::Type InCommitType, FProperty* InProperty, TArray<TSharedPtr<FBindingChainElement>> InBindingChain);
 
 	void HandleCreateAndAddBinding();
 
-	UStruct* ResolveIndirection(const TArray<TSharedPtr<FBindingChainElement>>& BindingChain) const;
+	static UStruct* ResolveIndirection(const FPropertyBindingWidgetArgs& Args, const TArray<TSharedPtr<FBindingChainElement>>& BindingChain);
 
 	EVisibility GetGotoBindingVisibility() const;
 
 	FReply HandleGotoBindingClicked();
 
 	// Helper function to call the OnCanAcceptProperty* delegates, handles conversion of binding chain to TConstArrayView<FBindingChainElement> as expected by the delegate.
-	bool CanAcceptPropertyOrChildren(FProperty* InProperty, TConstArrayView<TSharedPtr<FBindingChainElement>> InBindingChain) const;
+	static bool CanAcceptPropertyOrChildren(const FPropertyBindingWidgetArgs& Args, FProperty* InProperty, TConstArrayView<TSharedPtr<FBindingChainElement>> InBindingChain);
 	
 	// Helper function to call the OnCanBindProperty* delegates, handles conversion of binding chain to TConstArrayView<FBindingChainElement> as expected by the delegate.
-	bool CanBindProperty(FProperty* InProperty, TConstArrayView<TSharedPtr<FBindingChainElement>> InBindingChain) const;
+	static bool CanBindProperty(const FPropertyBindingWidgetArgs& Args, FProperty* InProperty, TConstArrayView<TSharedPtr<FBindingChainElement>> InBindingChain);
 
 	virtual FReply OnDrop(const FGeometry& MyGeometry, const FDragDropEvent& DragDropEvent) override;
 
@@ -91,20 +93,20 @@ protected:
 
 private:
 
-	bool IsClassDenied(UClass* OwnerClass) const;
-	bool IsFieldFromDeniedClass(FFieldVariant Field) const;
-	bool HasBindableProperties(UStruct* InStruct, TArray<TSharedPtr<FBindingChainElement>>& BindingChain) const;
-	bool HasBindablePropertiesRecursive(UStruct* InStruct, TSet<UStruct*>& VisitedStructs, TArray<TSharedPtr<FBindingChainElement>>& BindingChain) const;
+	static bool IsClassDenied(const FPropertyBindingWidgetArgs& Args, UClass* OwnerClass);
+	static bool IsFieldFromDeniedClass(const FPropertyBindingWidgetArgs& Args, FFieldVariant Field);
+	static bool HasBindableProperties(const FPropertyBindingWidgetArgs& Args, const UBlueprint* Blueprint, UStruct* InStruct, TArray<TSharedPtr<FBindingChainElement>>& BindingChain);
+	static bool HasBindablePropertiesRecursive(const FPropertyBindingWidgetArgs& Args, const UBlueprint* Blueprint, UStruct* InStruct, TSet<UStruct*>& VisitedStructs, TArray<TSharedPtr<FBindingChainElement>>& BindingChain);
 
 	/**
 	 * Note that an ArrayView is not used to pass the BindingChain around since the predicate can modify the array
 	 * and this will invalidate the ArrayView if reallocation is performed.
 	 */
 	template <typename Predicate>
-	void ForEachBindableProperty(UStruct* InStruct, const TArray<TSharedPtr<FBindingChainElement>>& BindingChain, Predicate Pred) const;
+	static void ForEachBindableProperty(const FPropertyBindingWidgetArgs& Args,  const UBlueprint* Blueprint, UStruct* InStruct, const TArray<TSharedPtr<FBindingChainElement>>& BindingChain, Predicate Pred);
 
 	template <typename Predicate>
-	void ForEachBindableFunction(UClass* FromClass, Predicate Pred) const;
+	static void ForEachBindableFunction(const FPropertyBindingWidgetArgs& Args, UClass* FromClass, Predicate Pred);
 
 	void AddCategoryToMenu(FMenuBuilder& MenuBuilder, const FBindingContextStructCategory& Category);
 	void BuildContextStructCategoryRecursive(TConstArrayView<FString> CategoryNames, TArray<FBindingContextStructCategory>& ParentSubCategories, int32 ContextStructIndex);
diff --git a/Engine/Source/Editor/UnrealEd/Public/IPropertyAccessEditor.h b/Engine/Source/Editor/UnrealEd/Public/IPropertyAccessEditor.h
--- a/Engine/Source/Editor/UnrealEd/Public/IPropertyAccessEditor.h
+++ b/Engine/Source/Editor/UnrealEd/Public/IPropertyAccessEditor.h
@@ -17,6 +17,7 @@ class IPropertyHandle;
 class UEdGraph;
 class FExtender;
 class SWidget;
+class FMenuBuilder;
 struct FSlateBrush;
 struct FButtonStyle;
 struct FEdGraphPinType;
@@ -410,6 +411,17 @@ public:
 	 */
 	virtual TSharedRef<SWidget> MakePropertyBindingWidget(UBlueprint* InBlueprint, const FPropertyBindingWidgetArgs& InArgs = FPropertyBindingWidgetArgs()) const = 0;
 
+	/** 
+	 * EXPERIMENTAL: Make a property binding menu for a struct.
+	 * @param	MenuBuilder					The menu to add entries to
+	 * @param	InArgs						arguments for filtering properties etc. 
+	 * @param	InPropertyName				The Property Name to pass to the OnAddBinding callback
+	 * @param	InOwnerStruct				The Struct/Class to get properties from
+	 * @param	InBindingChain				The base binding chain to InOwnerStruct
+	 */
+	virtual void FillPropertyMenu(FMenuBuilder& MenuBuilder, const FPropertyBindingWidgetArgs& InArgs, FName InPropertyName, UStruct* InOwnerStruct,
+										TArray<TSharedPtr<FBindingChainElement>> InBindingChain) const = 0;
+	
 	/**
 	 * Make a property binding widget.
 	 * @param	InBindingContextStructs		An array of structs the binding will exist within
