Fix PropertyName captured by reference

#code-analyzed

diff --git a/Engine/Plugins/Runtime/PropertyAccess/Source/PropertyAccessEditor/Private/SPropertyBinding.cpp b/Engine/Plugins/Runtime/PropertyAccess/Source/PropertyAccessEditor/Private/SPropertyBinding.cpp
--- a/Engine/Plugins/Runtime/PropertyAccess/Source/PropertyAccessEditor/Private/SPropertyBinding.cpp
+++ b/Engine/Plugins/Runtime/PropertyAccess/Source/PropertyAccessEditor/Private/SPropertyBinding.cpp
@@ -875,7 +875,7 @@ void SPropertyBinding::FillPropertyMenu(FMenuBuilder& MenuBuilder, const FProper
 
 			MenuBuilder.BeginSection("Functions", LOCTEXT("Functions", "Functions"));
 			{
-				ForEachBindableFunction(Args, OwnerClass, [&Args, Blueprint, &PropertyName, &InBindingChain, &MenuBuilder, &MakeFunctionWidget] (const FFunctionInfo& Info) 
+				ForEachBindableFunction(Args, OwnerClass, [&Args, Blueprint, PropertyName, &InBindingChain, &MenuBuilder, &MakeFunctionWidget] (const FFunctionInfo& Info) 
 				{
 					TArray<TSharedPtr<FBindingChainElement>> NewBindingChain(InBindingChain);
 					NewBindingChain.Emplace(MakeShared<FBindingChainElement>(Info.Function));
@@ -903,7 +903,7 @@ void SPropertyBinding::FillPropertyMenu(FMenuBuilder& MenuBuilder, const FProper
 								MenuBuilder.AddSubMenu(
 									MakeFunctionWidget(Info),
 									FNewMenuDelegate::CreateLambda(
-										[Args, Blueprint, &PropertyName, ObjectPropertyBase, NewBindingChain](FMenuBuilder& MenuBuilder)
+										[Args, Blueprint, PropertyName, ObjectPropertyBase, NewBindingChain](FMenuBuilder& MenuBuilder)
 											{
 												FillPropertyMenu(MenuBuilder, Args, Blueprint, PropertyName, static_cast<UStruct*>(ObjectPropertyBase->PropertyClass), NewBindingChain);
 											}));
@@ -915,7 +915,7 @@ void SPropertyBinding::FillPropertyMenu(FMenuBuilder& MenuBuilder, const FProper
 							{
 								MenuBuilder.AddSubMenu(
 									MakeFunctionWidget(Info),
-									FNewMenuDelegate::CreateLambda([Args, Blueprint, &PropertyName, StructProperty, NewBindingChain](FMenuBuilder& MenuBuilder)
+									FNewMenuDelegate::CreateLambda([Args, Blueprint, PropertyName, StructProperty, NewBindingChain](FMenuBuilder& MenuBuilder)
 								{
 									FillPropertyMenu(MenuBuilder, Args, Blueprint, PropertyName, static_cast<UStruct*>(StructProperty->Struct), NewBindingChain);
 								}));
@@ -962,7 +962,7 @@ void SPropertyBinding::FillPropertyMenu(FMenuBuilder& MenuBuilder, const FProper
 
 			MenuBuilder.BeginSection("Properties", LOCTEXT("Properties", "Properties"));
 			{
-				ForEachBindableProperty(Args, Blueprint, InOwnerStruct, InBindingChain, [&Args, Blueprint, &PropertyName, &InBindingChain, BindingStruct, &MenuBuilder, &MakeArrayElementPropertyWidget, &MakePropertyWidget, &MakePropertyEntry, &MakeArrayElementEntry](FProperty* Property)
+				ForEachBindableProperty(Args, Blueprint, InOwnerStruct, InBindingChain, [&Args, Blueprint, PropertyName, &InBindingChain, BindingStruct, &MenuBuilder, &MakeArrayElementPropertyWidget, &MakePropertyWidget, &MakePropertyEntry, &MakeArrayElementEntry](FProperty* Property)
 					{
 						TArray<TSharedPtr<FBindingChainElement>> NewBindingChain(InBindingChain);
 						TArray<TSharedPtr<FBindingChainElement>> NewArrayElementBindingChain;
